Initial submission of livelink hub program launcher

diff --git a/Engine/Programs/LiveLinkHub/Config/DefaultEngine.ini b/Engine/Programs/LiveLinkHub/Config/DefaultEngine.ini
new file mode 100644
--- /dev/null
+++ b/Engine/Programs/LiveLinkHub/Config/DefaultEngine.ini
@@ -0,0 +1,15 @@
+[URL]
+GameNameShort=LLH
+
+[Plugins]
++ProgramEnabledPlugins="LiveLink"
++ProgramEnabledPlugins="StructUtils"
++ProgramEnabledPlugins="ContentBrowserAssetDataSource"
++ProgramEnabledPlugins="LiveLinkFreeD"
+
+[/Script/MacTargetPlatform.XcodeProjectSettings]
+TemplateMacPlist=(FilePath="/Engine/Source/Programs/LiveLinkHub/Resources/Mac/Info.plist")
+
+
+[/Script/UnrealEd.ProjectPackagingSettings]
+bShareMaterialShaderCode=false
diff --git a/Engine/Source/Programs/LiveLinkHub/LiveLinkHub.Target.cs b/Engine/Source/Programs/LiveLinkHub/LiveLinkHub.Target.cs
new file mode 100644
--- /dev/null
+++ b/Engine/Source/Programs/LiveLinkHub/LiveLinkHub.Target.cs
@@ -0,0 +1,45 @@
+// Copyright Epic Games, Inc. All Rights Reserved.
+
+using UnrealBuildTool;
+using System.Collections.Generic;
+
+[SupportedPlatforms(UnrealPlatformClass.Desktop)]
+public class LiveLinkHubTarget : TargetRules
+{
+	public LiveLinkHubTarget(TargetInfo Target) : base(Target)
+	{
+		Type = TargetType.Program;
+		LinkType = TargetLinkType.Modular;
+		LaunchModuleName = "LiveLinkHubLauncher";
+
+		IncludeOrderVersion = EngineIncludeOrderVersion.Latest;
+
+		SolutionDirectory = "Programs/LiveLink";
+
+		AdditionalPlugins.Add("LiveLink");
+		AdditionalPlugins.Add("ContentBrowserAssetDataSource");
+		AdditionalPlugins.Add("StructUtils");
+		AdditionalPlugins.Add("UdpMessaging");
+		AdditionalPlugins.Add("QuicMessaging");
+
+		bCompileAgainstCoreUObject = true;
+		bCompileAgainstEngine = true;
+		bCompileAgainstEditor = true;
+		bBuildWithEditorOnlyData = true;
+		bIncludePluginsForTargetPlatforms = false;
+		bUsesSlate = true;
+
+		bCompileWithPluginSupport = true;
+		bBuildDeveloperTools = true;
+		bIsBuildingConsoleApplication = false;
+
+		GlobalDefinitions.Add("WITH_LIVELINK_HUB=1");
+
+		bEnableTrace = false;
+
+		OptedInModulePlatforms = new UnrealTargetPlatform[] { UnrealTargetPlatform.Win64, UnrealTargetPlatform.Mac,
+															  UnrealTargetPlatform.Linux, UnrealTargetPlatform.LinuxArm64 };
+
+		// todo: Look into using ExeBinariesSubFolder
+	}
+}
diff --git a/Engine/Source/Programs/LiveLinkHub/LiveLinkHubLauncher.Build.cs b/Engine/Source/Programs/LiveLinkHub/LiveLinkHubLauncher.Build.cs
new file mode 100644
--- /dev/null
+++ b/Engine/Source/Programs/LiveLinkHub/LiveLinkHubLauncher.Build.cs
@@ -0,0 +1,78 @@
+// Copyright Epic Games, Inc. All Rights Reserved.
+
+using Microsoft.Extensions.Logging;
+using System;
+using UnrealBuildTool;
+using System.IO;
+
+// This is named LiveLinkHubLauncher in order to not clash with the LiveLinkHub module which is in the livelink plugin.
+public class LiveLinkHubLauncher : ModuleRules
+{
+	public LiveLinkHubLauncher(ReadOnlyTargetRules Target) : base(Target)
+	{
+		PrivateDependencyModuleNames.Add("LiveLinkHub");
+
+		// LaunchEngineLoop dependencies
+		PrivateDependencyModuleNames.AddRange(
+			new string[] {
+				"ApplicationCore",
+				"Core",
+				"CoreUObject",
+				"Engine",
+				"InputCore",
+				"InstallBundleManager",
+				"MediaUtils",
+				"Messaging",
+				"MoviePlayer",
+				"MoviePlayerProxy",
+				"Projects",
+				"PreLoadScreen",
+				"PIEPreviewDeviceProfileSelector",
+				"RenderCore",
+				"RHI",
+				"Slate",
+				"SlateCore",
+				"StandaloneRenderer",
+				"TraceLog", 
+			}
+		);
+
+		// LaunchEngineLoop IncludePath dependencies
+		PrivateIncludePathModuleNames.AddRange(
+			new string[] {
+				"Launch",
+				"AutomationWorker",
+				"AutomationController",
+				"AutomationTest",
+				"DerivedDataCache",
+				"HeadMountedDisplay", 
+				"MRMesh", 
+				"ProfileVisualizer", 
+				"ProfilerService", 
+				"SlateRHIRenderer", 
+				"SlateNullRenderer",
+			}
+		);
+
+		// LaunchEngineLoop editor dependencies
+		PrivateDependencyModuleNames.AddRange(
+			new string[] {
+				
+				"PropertyEditor",
+				"DerivedDataCache",
+				"ToolWidgets",
+				"UnrealEd"
+		});
+
+		if (Target.IsInPlatformGroup(UnrealPlatformGroup.Linux))
+		{
+			PrivateDependencyModuleNames.AddRange(
+				new string[] {
+					"UnixCommonStartup"
+				}
+			);
+		}
+
+		ShortName = "LLHub";
+	}
+}
diff --git a/Engine/Source/Programs/LiveLinkHub/Private/LiveLinkHubCommandlet.h b/Engine/Source/Programs/LiveLinkHub/Private/LiveLinkHubCommandlet.h
new file mode 100644
--- /dev/null
+++ b/Engine/Source/Programs/LiveLinkHub/Private/LiveLinkHubCommandlet.h
@@ -0,0 +1,40 @@
+// Copyright Epic Games, Inc. All Rights Reserved.
+
+#pragma once
+
+#include "Commandlets/Commandlet.h"
+#include "Modules/ModuleManager.h"
+#include "ILiveLinkHubModule.h"
+
+#include "LiveLinkHubCommandlet.generated.h"
+
+UCLASS()
+class ULiveLinkHubCommandlet : public UCommandlet
+{
+public:
+	GENERATED_BODY()
+
+	ULiveLinkHubCommandlet()
+	{
+		LogToConsole = false;
+		FastExit = true;
+	}
+	
+	/** Runs the commandlet */
+	virtual int32 Main(const FString& Params) override
+	{
+		PRIVATE_GAllowCommandletRendering = true;
+
+		// Work around the CB file asset data source being disabled in commandlets.
+		if (IModuleInterface* AssetDataSourceModule = FModuleManager::Get().GetModule("ContentBrowserAssetDataSource"))
+		{
+			PRIVATE_GIsRunningCommandlet = false;
+			AssetDataSourceModule->StartupModule();
+			PRIVATE_GIsRunningCommandlet = true;	
+		}
+		
+		FModuleManager::Get().LoadModuleChecked<ILiveLinkHubModule>("LiveLinkHub").StartLiveLinkHub();
+
+		return 0;
+	}
+};
diff --git a/Engine/Source/Programs/LiveLinkHub/Private/LiveLinkHubRun.cpp b/Engine/Source/Programs/LiveLinkHub/Private/LiveLinkHubRun.cpp
new file mode 100644
--- /dev/null
+++ b/Engine/Source/Programs/LiveLinkHub/Private/LiveLinkHubRun.cpp
@@ -0,0 +1,25 @@
+#include "LiveLinkHubRun.h"
+
+#include "LaunchEngineLoop.h"
+
+DEFINE_LOG_CATEGORY_STATIC(LogLiveLinkHub, Log, All);
+
+#define LOCTEXT_NAMESPACE "LiveLinkHub"
+
+int32 RunLiveLinkHub(const TCHAR* CommandLine)
+{
+	FTaskTagScope Scope(ETaskTag::EGameThread);
+	
+	// Start up the main loop, adding some extra command line arguments:
+	const int32 Result = GEngineLoop.PreInit(*FString::Printf(TEXT("%s %s"), CommandLine, TEXT("LiveLinkHubCommandlet -Messaging")));
+
+	if (Result != 0)
+	{
+		UE_LOG(LogLiveLinkHub, Error, TEXT("EngineLoop PreInit failed!"));
+		return Result;
+	}
+
+	GEngineLoop.Exit();
+
+	return Result;
+}
\ No newline at end of file
diff --git a/Engine/Source/Programs/LiveLinkHub/Private/LiveLinkHubRun.h b/Engine/Source/Programs/LiveLinkHub/Private/LiveLinkHubRun.h
new file mode 100644
--- /dev/null
+++ b/Engine/Source/Programs/LiveLinkHub/Private/LiveLinkHubRun.h
@@ -0,0 +1,7 @@
+// Copyright Epic Games, Inc. All Rights Reserved.
+
+#pragma once
+
+#include "HAL/Platform.h"
+
+int32 RunLiveLinkHub(const TCHAR* CommandLine);
diff --git a/Engine/Source/Programs/LiveLinkHub/Private/LiveLinkHubRunMain.cpp b/Engine/Source/Programs/LiveLinkHub/Private/LiveLinkHubRunMain.cpp
new file mode 100644
--- /dev/null
+++ b/Engine/Source/Programs/LiveLinkHub/Private/LiveLinkHubRunMain.cpp
@@ -0,0 +1,6 @@
+// Copyright Epic Games, Inc. All Rights Reserved.
+
+#include "HAL/Platform.h"
+#include "RequiredProgramMainCPPInclude.h"
+
+IMPLEMENT_APPLICATION(LiveLinkHubLauncher, "LiveLinkHub");
diff --git a/Engine/Source/Programs/LiveLinkHub/Private/Platform/Linux/LinuxLiveLinkHubLauncherMain.cpp b/Engine/Source/Programs/LiveLinkHub/Private/Platform/Linux/LinuxLiveLinkHubLauncherMain.cpp
new file mode 100644
--- /dev/null
+++ b/Engine/Source/Programs/LiveLinkHub/Private/Platform/Linux/LinuxLiveLinkHubLauncherMain.cpp
@@ -0,0 +1,11 @@
+// Copyright Epic Games, Inc. All Rights Reserved.
+
+#include "CoreMinimal.h"
+#include "LiveLinkHubRun.h"
+#include "LaunchEngineLoop.h"
+#include "UnixCommonStartup.h"
+
+int main(int argc, char *argv[])
+{
+	return CommonUnixMain(argc, argv, &RunLiveLinkHub, [] { } );
+}
diff --git a/Engine/Source/Programs/LiveLinkHub/Private/Platform/Mac/MacLiveLinkHubLauncherMain.cpp b/Engine/Source/Programs/LiveLinkHub/Private/Platform/Mac/MacLiveLinkHubLauncherMain.cpp
new file mode 100644
--- /dev/null
+++ b/Engine/Source/Programs/LiveLinkHub/Private/Platform/Mac/MacLiveLinkHubLauncherMain.cpp
@@ -0,0 +1,10 @@
+// Copyright Epic Games, Inc. All Rights Reserved.
+
+#include "LiveLinkHubRun.h"
+#include "Mac/MacProgramDelegate.h"
+#include "LaunchEngineLoop.h"
+
+int main( int argc, char *argv[] )
+{
+	return [MacProgramDelegate mainWithArgc:argc argv:argv programMain:RunLiveLinkHub programExit:FEngineLoop::AppExit];
+}
diff --git a/Engine/Source/Programs/LiveLinkHub/Private/Platform/Windows/WindowsLiveLinkHubLauncherMain.cpp b/Engine/Source/Programs/LiveLinkHub/Private/Platform/Windows/WindowsLiveLinkHubLauncherMain.cpp
new file mode 100644
--- /dev/null
+++ b/Engine/Source/Programs/LiveLinkHub/Private/Platform/Windows/WindowsLiveLinkHubLauncherMain.cpp
@@ -0,0 +1,64 @@
+// Copyright Epic Games, Inc. All Rights Reserved.
+
+#include "CoreMinimal.h"
+#include "LiveLinkHubRun.h"
+#include "HAL/ExceptionHandling.h"
+#include "LaunchEngineLoop.h"
+#include "Windows/WindowsHWrapper.h"
+#include "Misc/CommandLine.h"
+#include "Misc/OutputDeviceError.h"
+
+/**
+ * The main application entry point for Windows platforms.
+ *
+ * @param hInInstance Handle to the current instance of the application.
+ * @param hPrevInstance Handle to the previous instance of the application (always NULL).
+ * @param lpCmdLine Command line for the application.
+ * @param nShowCmd Specifies how the window is to be shown.
+ * @return Application's exit value.
+ */
+int32 WINAPI WinMain( HINSTANCE hInInstance, HINSTANCE hPrevInstance, char* lpCmdLine, int32 nShowCmd )
+{
+	hInstance = hInInstance;
+
+	const TCHAR* CmdLine = ::GetCommandLineW();
+	CmdLine = FCommandLine::RemoveExeName(CmdLine);
+
+#if !UE_BUILD_SHIPPING
+	if (FParse::Param(CmdLine, TEXT("crashreports")))
+	{
+		GAlwaysReportCrash = true;
+	}
+#endif
+
+	int32 ErrorLevel = 0;
+
+#if UE_BUILD_DEBUG
+	if (!GAlwaysReportCrash)
+#else
+	if (FPlatformMisc::IsDebuggerPresent() && !GAlwaysReportCrash)
+#endif
+	{
+		ErrorLevel = RunLiveLinkHub(CmdLine);
+	}
+	else
+	{
+#if !PLATFORM_SEH_EXCEPTIONS_DISABLED
+		__try
+#endif
+ 		{
+			GIsGuarded = 1;
+			ErrorLevel = RunLiveLinkHub(CmdLine);
+			GIsGuarded = 0;
+		}
+#if !PLATFORM_SEH_EXCEPTIONS_DISABLED
+		__except (ReportCrash(GetExceptionInformation()))
+		{
+			ErrorLevel = 1;
+			GError->HandleError();
+			FPlatformMisc::RequestExit(true);
+		}
+#endif
+	}
+	return ErrorLevel;
+}
diff --git a/Engine/Source/Programs/LiveLinkHub/Resources/Mac/Assets.xcassets/AppIcon.appiconset/Contents.json b/Engine/Source/Programs/LiveLinkHub/Resources/Mac/Assets.xcassets/AppIcon.appiconset/Contents.json
new file mode 100644
--- /dev/null
+++ b/Engine/Source/Programs/LiveLinkHub/Resources/Mac/Assets.xcassets/AppIcon.appiconset/Contents.json
@@ -0,0 +1,68 @@
+{
+  "images" : [
+    {
+      "filename" : "icon_16x16.png",
+      "idiom" : "mac",
+      "scale" : "1x",
+      "size" : "16x16"
+    },
+    {
+      "filename" : "[REDACTED_EMAIL]",
+      "idiom" : "mac",
+      "scale" : "2x",
+      "size" : "16x16"
+    },
+    {
+      "filename" : "icon_32x32.png",
+      "idiom" : "mac",
+      "scale" : "1x",
+      "size" : "32x32"
+    },
+    {
+      "filename" : "[REDACTED_EMAIL]",
+      "idiom" : "mac",
+      "scale" : "2x",
+      "size" : "32x32"
+    },
+    {
+      "filename" : "icon_128x128.png",
+      "idiom" : "mac",
+      "scale" : "1x",
+      "size" : "128x128"
+    },
+    {
+      "filename" : "[REDACTED_EMAIL]",
+      "idiom" : "mac",
+      "scale" : "2x",
+      "size" : "128x128"
+    },
+    {
+      "filename" : "icon_256x256.png",
+      "idiom" : "mac",
+      "scale" : "1x",
+      "size" : "256x256"
+    },
+    {
+      "filename" : "[REDACTED_EMAIL]",
+      "idiom" : "mac",
+      "scale" : "2x",
+      "size" : "256x256"
+    },
+    {
+      "filename" : "icon_512x512.png",
+      "idiom" : "mac",
+      "scale" : "1x",
+      "size" : "512x512"
+    },
+    {
+      "filename" : "[REDACTED_EMAIL]",
+      "idiom" : "mac",
+      "scale" : "2x",
+      "size" : "512x512"
+    }
+  ],
+  "info" : {
+    "author" : "xcode",
+    "version" : 1
+  }
+}
diff --git a/Engine/Source/Programs/LiveLinkHub/Resources/Mac/Assets.xcassets/Contents.json b/Engine/Source/Programs/LiveLinkHub/Resources/Mac/Assets.xcassets/Contents.json
new file mode 100644
--- /dev/null
+++ b/Engine/Source/Programs/LiveLinkHub/Resources/Mac/Assets.xcassets/Contents.json
@@ -0,0 +1,6 @@
+{
+  "info" : {
+    "author" : "xcode",
+    "version" : 1
+  }
+}
diff --git a/Engine/Source/Programs/LiveLinkHub/Resources/Windows/LiveLinkHub.rc b/Engine/Source/Programs/LiveLinkHub/Resources/Windows/LiveLinkHub.rc
new file mode 100644
--- /dev/null
+++ b/Engine/Source/Programs/LiveLinkHub/Resources/Windows/LiveLinkHub.rc
@@ -0,0 +1,88 @@
+// Copyright Epic Games, Inc. All Rights Reserved.
+
+#define APSTUDIO_READONLY_SYMBOLS
+/////////////////////////////////////////////////////////////////////////////
+//
+// Generated from the TEXTINCLUDE 2 resource.
+//
+#include <windows.h>
+#include "Runtime/Launch/Resources/Version.h"
+#include "Runtime/Launch/Resources/Windows/resource.h"
+
+#define LIVELINK_HUB_VERSION_MAJOR 1
+#define LIVELINK_HUB_VERSION_MINOR 0
+#define LIVELINK_HUB_VERSION_PATCH 0
+
+#define LIVELINK_HUB_VERSION_STRING "v1.01"
+
+#define VER_FILEVERSION         LIVELINK_HUB_VERSION_MAJOR,LIVELINK_HUB_VERSION_MINOR,LIVELINK_HUB_VERSION_PATCH,0
+#define VER_FILEVERSION_STR     LIVELINK_HUB_VERSION_STRING "\0"
+
+#define VER_PRODUCTVERSION      ENGINE_MAJOR_VERSION,ENGINE_MINOR_VERSION,ENGINE_PATCH_VERSION,0
+#define VER_PRODUCTVERSION_STR  ENGINE_VERSION_STRING "\0"
+
+#ifdef _DEBUG
+#define VER_FILEFLAGS           0x1L // VS_FF_DEBUG
+#else
+#define VER_FILEFLAGS           0x0L
+#endif
+
+#define IDI_LIVELINK_HUB            IDICON_UEGame // 123
+
+/////////////////////////////////////////////////////////////////////////////
+#undef APSTUDIO_READONLY_SYMBOLS
+
+/////////////////////////////////////////////////////////////////////////////
+// English (U.S.) resources
+
+#if !defined(AFX_RESOURCE_DLL) || defined(AFX_TARG_ENU)
+#ifdef _WIN32
+LANGUAGE LANG_ENGLISH, SUBLANG_ENGLISH_US
+#pragma code_page(1252)
+#endif //_WIN32
+
+/////////////////////////////////////////////////////////////////////////////
+//
+// Version
+//
+
+VS_VERSION_INFO VERSIONINFO
+FILEVERSION     VER_FILEVERSION
+PRODUCTVERSION  VER_PRODUCTVERSION
+FILEFLAGSMASK   0x17L
+FILEFLAGS       VER_FILEFLAGS
+FILEOS          0x4L
+FILETYPE        0x2L
+FILESUBTYPE     0x0L
+BEGIN
+    BLOCK "StringFileInfo"
+    BEGIN
+        BLOCK "040904b0"
+        BEGIN
+            VALUE "CompanyName",      EPIC_COMPANY_NAME
+            VALUE "LegalCopyright",   EPIC_COPYRIGHT_STRING
+            VALUE "ProductName",      EPIC_PRODUCT_NAME
+            VALUE "ProductVersion",   VER_PRODUCTVERSION_STR
+            VALUE "FileDescription",  "LiveLink Hub"
+            VALUE "FileVersion",      VER_FILEVERSION_STR
+            VALUE "InternalName",     "LiveLinkHub"
+            VALUE "OriginalFilename", "LiveLinkHub.exe"
+        END
+    END
+    BLOCK "VarFileInfo"
+    BEGIN
+        VALUE "Translation", 0x409, 1200
+    END
+END
+
+/////////////////////////////////////////////////////////////////////////////
+//
+// Icon
+//
+
+// Icon with lowest ID value placed first to ensure application icon
+// remains consistent on all systems.
+IDI_LIVELINK_HUB    ICON    "LiveLinkHub.ico"
+
+#endif    // English (U.S.) resources
+/////////////////////////////////////////////////////////////////////////////
diff --git a/Engine/Source/Programs/LiveLinkHub/Resources/Windows/make_ico.py b/Engine/Source/Programs/LiveLinkHub/Resources/Windows/make_ico.py
new file mode 100644
--- /dev/null
+++ b/Engine/Source/Programs/LiveLinkHub/Resources/Windows/make_ico.py
@@ -0,0 +1,15 @@
+# Copyright Epic Games, Inc. All Rights Reserved.
+
+from PIL import Image
+
+def MakeIcon(folder, output):
+    sizes = [256, 48, 32, 24, 16]
+    images = []
+    for size in sizes:
+        print(f"Loading {folder}/icon_{size}x{size}.png")
+        image = Image.open(f"{folder}/icon_{size}x{size}.png")
+        images.append(image)
+    print(f"Creating {output}")
+    images[0].save(f"{output}", append_images=images)
+
+MakeIcon("LiveLinkHub", "LiveLinkHub.ico")
