Fix crash when loading livelink in shipping outside of engine directory
Compile without python (Note that plugin dependency is still needed for the other plugins to load)
Compile with MessageBus support in shipping
Compile without localization support

diff --git a/Engine/Plugins/Animation/LiveLink/Source/LiveLinkEditor/Private/LiveLinkPanelController.cpp b/Engine/Plugins/Animation/LiveLink/Source/LiveLinkEditor/Private/LiveLinkPanelController.cpp
--- a/Engine/Plugins/Animation/LiveLink/Source/LiveLinkEditor/Private/LiveLinkPanelController.cpp
+++ b/Engine/Plugins/Animation/LiveLink/Source/LiveLinkEditor/Private/LiveLinkPanelController.cpp
@@ -21,6 +21,7 @@ FLiveLinkPanelController::FLiveLinkPanelController(TAttribute<bool> bInReadOnly)
 	OnSourcesChangedHandle = Client->OnLiveLinkSourcesChanged().AddRaw(this, &FLiveLinkPanelController::OnSourcesChangedHandler);
 	OnSubjectsChangedHandle = Client->OnLiveLinkSubjectsChanged().AddRaw(this, &FLiveLinkPanelController::OnSubjectsChangedHandler);
 
+	FLiveLinkClientCommands::Register();
 	CommandList = MakeShared<FUICommandList>();
 	BindCommands();
 
diff --git a/Engine/Source/Programs/LiveLinkHub/LiveLinkHub.Target.cs b/Engine/Source/Programs/LiveLinkHub/LiveLinkHub.Target.cs
--- a/Engine/Source/Programs/LiveLinkHub/LiveLinkHub.Target.cs
+++ b/Engine/Source/Programs/LiveLinkHub/LiveLinkHub.Target.cs
@@ -46,13 +46,17 @@ public class LiveLinkHubTarget : TargetRules
 
 		bUsesSlate = true;
 
+		bCompileICU = false;
+		bCompilePython = false; 
 		bCompileWithPluginSupport = true;
 		bBuildDeveloperTools = true;
 		bIsBuildingConsoleApplication = false;
 
+
 		GlobalDefinitions.Add("WITH_LIVELINK_HUB=1");
 		GlobalDefinitions.Add("AUTOSDKS_ENABLED=0");
 		GlobalDefinitions.Add("NO_LOGGING=0");
+		GlobalDefinitions.Add("PLATFORM_SUPPORTS_MESSAGEBUS=1");
 
 		bEnableTrace = true;
 
diff --git a/Engine/Source/Programs/LiveLinkHub/Private/LiveLinkHubCommandlet.h b/Engine/Source/Programs/LiveLinkHub/Private/LiveLinkHubCommandlet.h
--- a/Engine/Source/Programs/LiveLinkHub/Private/LiveLinkHubCommandlet.h
+++ b/Engine/Source/Programs/LiveLinkHub/Private/LiveLinkHubCommandlet.h
@@ -3,6 +3,7 @@
 #pragma once
 
 #include "Commandlets/Commandlet.h"
+#include "Misc/CoreDelegates.h"
 #include "Modules/ModuleManager.h"
 #include "ILiveLinkHubModule.h"
 
@@ -34,7 +35,8 @@ public:
 		}
 
 		FModuleManager::Get().LoadModule("OutputLog");
-		
+
+		FModuleManager::Get().LoadModuleChecked<ILiveLinkHubModule>("LiveLinkEditor");
 		FModuleManager::Get().LoadModuleChecked<ILiveLinkHubModule>("LiveLinkHub").StartLiveLinkHub();
 
 		return 0;
diff --git a/Engine/Source/Programs/LiveLinkHub/Private/LiveLinkHubRun.cpp b/Engine/Source/Programs/LiveLinkHub/Private/LiveLinkHubRun.cpp
--- a/Engine/Source/Programs/LiveLinkHub/Private/LiveLinkHubRun.cpp
+++ b/Engine/Source/Programs/LiveLinkHub/Private/LiveLinkHubRun.cpp
@@ -1,15 +1,26 @@
 // Copyright Epic Games, Inc. All Rights Reserved.
 
 #include "LiveLinkHubRun.h"
-
+#include "Misc/CoreDelegates.h"
 #include "LaunchEngineLoop.h"
 
+#ifndef WITH_ASSET_LOADING_AUDIT
+#define WITH_ASSET_LOADING_AUDIT 0
+#endif
+
 DEFINE_LOG_CATEGORY_STATIC(LogLiveLinkHub, Log, All);
 
 int32 RunLiveLinkHub(const TCHAR* CommandLine)
 {
 	FTaskTagScope Scope(ETaskTag::EGameThread);
-	
+
+#if WITH_ASSET_LOADING_AUDIT
+	FCoreDelegates::OnSyncLoadPackage.AddLambda([](const FString& PackageName)
+		{
+			UE_LOG(LogLiveLinkHub, Display, TEXT("Audit: Loaded %s"), *PackageName);
+		});
+#endif
+
 	// Start up the main loop, adding some extra command line arguments:
 	const int32 Result = GEngineLoop.PreInit(*FString::Printf(TEXT("%s %s"), CommandLine, TEXT("LiveLinkHubCommandlet -Messaging -DDC=NoShared -NoShaderCompile")));
 
