LiveLinkHub - Remove console window when launching app.

LLH contains an output log so we have no need for a console window anymore.

diff --git a/Engine/Plugins/Animation/LiveLinkHub/Source/LiveLinkHub/Private/LiveLinkHub.cpp b/Engine/Plugins/Animation/LiveLinkHub/Source/LiveLinkHub/Private/LiveLinkHub.cpp
--- a/Engine/Plugins/Animation/LiveLinkHub/Source/LiveLinkHub/Private/LiveLinkHub.cpp
+++ b/Engine/Plugins/Animation/LiveLinkHub/Source/LiveLinkHub/Private/LiveLinkHub.cpp
@@ -30,6 +30,9 @@ void FLiveLinkHub::Initialize()
 {
 	TRACE_CPUPROFILER_EVENT_SCOPE(FLiveLinkHub::Initialize);
 
+	// Re-enable this since we've disabled it to avoid the creation of the console window.
+	GIsSilent = false;
+
 	// We must register the livelink client first since we might rely on the modular feature to initialize the controllers/managers.
 	LiveLinkHubClient = MakeShared<FLiveLinkHubClient>(AsShared());
 	IModularFeatures::Get().RegisterModularFeature(ILiveLinkClient::ModularFeatureName, LiveLinkHubClient.Get());
diff --git a/Engine/Source/Programs/LiveLinkHub/LiveLinkHub.Target.cs b/Engine/Source/Programs/LiveLinkHub/LiveLinkHub.Target.cs
--- a/Engine/Source/Programs/LiveLinkHub/LiveLinkHub.Target.cs
+++ b/Engine/Source/Programs/LiveLinkHub/LiveLinkHub.Target.cs
@@ -58,10 +58,10 @@ public class LiveLinkHubTarget : TargetRules
 		bCompileWithPluginSupport = true;
 		bBuildDeveloperTools = true;
 		bIsBuildingConsoleApplication = false;
+		bUseLoggingInShipping = true;
 
 		GlobalDefinitions.Add("WITH_LIVELINK_HUB=1");
 		GlobalDefinitions.Add("AUTOSDKS_ENABLED=0");
-		GlobalDefinitions.Add("NO_LOGGING=0");
 		GlobalDefinitions.Add("PLATFORM_SUPPORTS_MESSAGEBUS=1");
 
 		if (bCookedEditorDistribution)
@@ -75,7 +75,5 @@ public class LiveLinkHubTarget : TargetRules
 
 		OptedInModulePlatforms = new UnrealTargetPlatform[] { UnrealTargetPlatform.Win64, UnrealTargetPlatform.Mac,
 															  UnrealTargetPlatform.Linux, UnrealTargetPlatform.LinuxArm64 };
-
-		// todo: Look into using ExeBinariesSubFolder
 	}
 }
diff --git a/Engine/Source/Programs/LiveLinkHub/Private/LiveLinkHubRun.cpp b/Engine/Source/Programs/LiveLinkHub/Private/LiveLinkHubRun.cpp
--- a/Engine/Source/Programs/LiveLinkHub/Private/LiveLinkHubRun.cpp
+++ b/Engine/Source/Programs/LiveLinkHub/Private/LiveLinkHubRun.cpp
@@ -70,6 +70,9 @@ int32 RunLiveLinkHub(const TCHAR* CommandLine)
 	IPluginManager::Get().AddPluginSearchPath(RemappedPlugins, /*bShouldRefresh*/ true);
 #endif
 
+	// Disable this when going through PreInit to prevent the console window from appearing.
+	GIsSilent = true;
+
 	// Start up the main loop, adding some extra command line arguments:
 	const int32 Result = GEngineLoop.PreInit(*FString::Printf(TEXT("%s %s"), CommandLine, TEXT("LiveLinkHubCommandlet -Messaging -DDC=NoShared -NoShaderCompile")));
 
