Remove MSQuic dependency from LiveLinkHub

diff --git a/Engine/Source/Programs/LiveLinkHub/LiveLinkHub.Target.cs b/Engine/Source/Programs/LiveLinkHub/LiveLinkHub.Target.cs
--- a/Engine/Source/Programs/LiveLinkHub/LiveLinkHub.Target.cs
+++ b/Engine/Source/Programs/LiveLinkHub/LiveLinkHub.Target.cs
@@ -44,7 +44,6 @@ public class LiveLinkHubTarget : TargetRules
 			"ProceduralMeshComponent", // Needed by LensComponent
 			"PropertyAccessEditor",
 			"PythonScriptPlugin",
-			"QuicMessaging",
 			"StructUtils",
 			"UdpMessaging"
 		});
