Adding CaptureManager plugin to the LiveLink Hub program

diff --git a/Engine/Plugins/Experimental/ControlFlows/ControlFlows.uplugin b/Engine/Plugins/Experimental/ControlFlows/ControlFlows.uplugin
--- a/Engine/Plugins/Experimental/ControlFlows/ControlFlows.uplugin
+++ b/Engine/Plugins/Experimental/ControlFlows/ControlFlows.uplugin
@@ -14,6 +14,7 @@
 	"IsBetaVersion": false,
 	"IsExperimentalVersion": false,
 	"Installed": false,
+	"SupportedPrograms": [ "LiveLinkHub" ],
 	"Modules": [
 		{
 			"Name": "ControlFlows",
diff --git a/Engine/Source/Programs/LiveLinkHub/LiveLinkHub.Target.cs b/Engine/Source/Programs/LiveLinkHub/LiveLinkHub.Target.cs
--- a/Engine/Source/Programs/LiveLinkHub/LiveLinkHub.Target.cs
+++ b/Engine/Source/Programs/LiveLinkHub/LiveLinkHub.Target.cs
@@ -20,6 +20,10 @@ public class LiveLinkHubTarget : TargetRules
 	[CommandLine("-EnableThirdPartyPlugins=")]
 	public bool bEnableThirdPartyPlugins = true;
 
+	// Whether to disable building Capture Manager plugin.
+	[CommandLine("-EnableCaptureManagerPlugin=")]
+	public bool bEnableCaptureManagerPlugin = false;
+
 	public LiveLinkHubTarget(TargetInfo Target) : base(Target)
 	{
 		Type = TargetType.Program;
@@ -60,6 +64,14 @@ public class LiveLinkHubTarget : TargetRules
 			});
 		}
 
+		if (bEnableCaptureManagerPlugin)
+		{
+			OptionalPlugins.AddRange(new string[]
+			{
+				"CaptureManager"
+			});
+		}
+
 		OptionalPlugins.AddRange(new string[]
 		{
 			"AppleARKitFaceSupport"
