LiveLink Hub: Editor target + staging.

Adds support for building a version of LiveLink Hub configured as TargetType.Editor, and support for staging it with a custom subclass of MakeCookedEditor.

By sharing a build environment with the UnrealEditor target, we will be able to create a standalone distribution of LiveLink Hub that is ABI-compatible with plugins built and packaged from the standard UE distribution. This achieves our goal of supporting dynamically loaded third-party vendor plugins.

This also reduces build and iteration times by sharing modules with UE, and once it replaces the monolithic program target, will reduce the size of the LiveLinkHub executable in the UE distribution from ~237 MB to ~1 MB.

Example staging invocation: `RunUAT MakeLiveLinkHubEditor -project=Engine\Source\Programs\LiveLinkHubEditor\LiveLinkHubEditor.uproject -build -skipcook -stage [-nodebuginfo]`.

(This CL additionally introduces support for staging the existing program target with `RunUAT BuildCookRun -project=Engine\Programs\LiveLinkHub\LiveLinkHub.uproject -build -skipcook -stage [-nodebuginfo]`; however, it requires one additional change to further restrict `OptedInModulePlatforms` in LiveLinkHub.Target.cs, currently commented out due to uncertain impact to existing build automation.)

diff --git a/Engine/Plugins/Runtime/PropertyAccess/PropertyAccessEditor.uplugin b/Engine/Plugins/Runtime/PropertyAccess/PropertyAccessEditor.uplugin
--- a/Engine/Plugins/Runtime/PropertyAccess/PropertyAccessEditor.uplugin
+++ b/Engine/Plugins/Runtime/PropertyAccess/PropertyAccessEditor.uplugin
@@ -23,5 +23,5 @@
 			"LoadingPhase" : "PreDefault"
 		}		
 	],
-	"SupportedPrograms": ["ChaosVisualDebugger"]
-}
\ No newline at end of file
+	"SupportedPrograms": ["ChaosVisualDebugger", "LiveLinkHub"]
+}
diff --git a/Engine/Programs/LiveLinkHub/LiveLinkHub.uproject b/Engine/Programs/LiveLinkHub/LiveLinkHub.uproject
new file mode 100644
--- /dev/null
+++ b/Engine/Programs/LiveLinkHub/LiveLinkHub.uproject
@@ -0,0 +1,11 @@
+{
+	"FileVersion": 3,
+	"EngineAssociation": "",
+	"Category": "",
+	"Description": "",
+	"TargetPlatforms": [
+		"Linux",
+		"Mac",
+		"Windows"
+	]
+}
diff --git a/Engine/Source/Programs/LiveLinkHub/LiveLinkHub.Automation.cs b/Engine/Source/Programs/LiveLinkHub/LiveLinkHub.Automation.cs
new file mode 100644
--- /dev/null
+++ b/Engine/Source/Programs/LiveLinkHub/LiveLinkHub.Automation.cs
@@ -0,0 +1,64 @@
+// Copyright Epic Games, Inc. All Rights Reserved.
+
+using EpicGames.Core;
+using System.Linq;
+
+
+namespace AutomationTool
+{
+	class MakeLiveLinkHubEditor : MakeCookedEditor
+	{
+		protected override ProjectParams MakeParams(string DLCName, string BasedOnReleaseVersion)
+		{
+			return new ProjectParams(
+				Command: this
+				, RawProjectPath: ProjectFile
+
+				, NoBootstrapExe: true
+				, DLCName: DLCName
+				, BasedOnReleaseVersion: BasedOnReleaseVersion
+				, DedicatedServer: bIsCookedCooker
+				, NoClient: bIsCookedCooker
+				, OptionalContent: true
+				, ClientCookedTargets: new ParamList<string>() // Prevent AutodetectSettings from looking for a game target
+				, EditorTargets: new ParamList<string>("LiveLinkHubEditor")
+			);
+		}
+
+		protected override void ModifyParams(ProjectParams BuildParams)
+		{
+			base.ModifyParams(BuildParams);
+
+			// We don't want the SDK dir / CookerSupportFiles
+			BuildParams.CookerSupportFilesSubdirectory = null;
+		}
+
+		protected override void ModifyDeploymentContext(ProjectParams Params, DeploymentContext SC)
+		{
+			ModifyStageContext Context = CreateContext(Params, SC);
+
+			DefaultModifyDeploymentContext(Params, SC, Context);
+
+			Context.Apply(SC);
+
+			// Copy .target receipt from project bin to engine bin
+			string PlatName = SC.StageTargetPlatform.PlatformType.ToString();
+			SC.FilesToStage.NonUFSFiles.Add(
+				new StagedFileReference($"Engine/Binaries/{PlatName}/{Context.ProjectName}.target"),
+				new FileReference($"Engine/Binaries/{PlatName}/{Context.ProjectName}.target"));
+
+			// Copy program DefaultEngine.ini to Engine/Config
+			StagedFileReference DefaultEngineDest = new StagedFileReference("Engine/Config/DefaultEngine.ini");
+			SC.FilesToStage.UFSFiles[DefaultEngineDest] = new FileReference($"Engine/Source/Programs/{Context.ProjectName}/Config/DefaultEngine.ini");
+
+			// Remove asset registry entry
+			SC.FilesToStage.UFSFiles.Remove(new StagedFileReference($"{Context.ProjectName}/EditorClientAssetRegistry.bin"));
+
+			// Move PDBs to debug (FIXME: should this be necessary? otherwise -nodebuginfo doesn't exclude them)
+			SC.FilesToStage.NonUFSDebugFiles.Union(SC.FilesToStage.NonUFSFiles.Where(x => x.Key.HasExtension(".pdb")));
+
+			// Remove any files from NonUFS that were also added to Debug
+			SC.FilesToStage.NonUFSFiles = SC.FilesToStage.NonUFSFiles.Where(x => !SC.FilesToStage.NonUFSDebugFiles.ContainsKey(x.Key)).ToDictionary(x => x.Key, x => x.Value);
+		}
+	}
+}
diff --git a/Engine/Source/Programs/LiveLinkHub/LiveLinkHub.Automation.csproj b/Engine/Source/Programs/LiveLinkHub/LiveLinkHub.Automation.csproj
new file mode 100644
--- /dev/null
+++ b/Engine/Source/Programs/LiveLinkHub/LiveLinkHub.Automation.csproj
@@ -0,0 +1,47 @@
+<Project Sdk="Microsoft.NET.Sdk">
+  <Import Project="..\Shared\UnrealEngine.csproj.props" />
+
+  <PropertyGroup>
+    <TargetFramework>net6.0</TargetFramework>
+    <Configuration Condition=" '$(Configuration)' == '' ">Development</Configuration>
+    <OutputType>Library</OutputType>
+    <AppendTargetFrameworkToOutputPath>false</AppendTargetFrameworkToOutputPath>
+    <GenerateAssemblyInfo>false</GenerateAssemblyInfo>
+    <GenerateTargetFrameworkAttribute>false</GenerateTargetFrameworkAttribute>
+    <Configurations>Debug;Release;Development</Configurations>
+    <RootNamespace>LiveLinkHub.Automation</RootNamespace>
+    <AssemblyName>LiveLinkHub.Automation</AssemblyName>
+    <WarningsNotAsErrors>612,618</WarningsNotAsErrors>
+    <OutputPath>..\..\Binaries\DotNET\AutomationTool\AutomationScripts\LiveLinkHub</OutputPath>
+    <TreatWarningsAsErrors>true</TreatWarningsAsErrors>
+    <DebugType>pdbonly</DebugType>
+  </PropertyGroup>
+  <PropertyGroup Condition=" '$(Configuration)|$(Platform)' == 'Development|AnyCPU' ">
+    <DefineConstants>$(DefineConstants);TRACE</DefineConstants>
+    <Optimize>true</Optimize>
+  </PropertyGroup>
+  <PropertyGroup Condition=" '$(Configuration)|$(Platform)' == 'Release|AnyCPU' ">
+    <DefineConstants>$(DefineConstants);TRACE</DefineConstants>
+    <Optimize>true</Optimize>
+  </PropertyGroup>
+  <PropertyGroup Condition=" '$(Configuration)|$(Platform)' == 'Debug|AnyCPU' ">
+    <DefineConstants>$(DefineConstants);DEBUG;TRACE</DefineConstants>
+    <DebugSymbols>true</DebugSymbols>
+    <DebugType>full</DebugType>
+  </PropertyGroup>
+
+  <ItemGroup>
+    <Compile Include="..\Shared\MetaData.cs">
+      <Link>Properties\MetaData.cs</Link>
+    </Compile>
+  </ItemGroup>
+  <ItemGroup>
+    <ProjectReference Include="..\AutomationTool\CookedEditor\CookedEditor.Automation.csproj" />
+    <ProjectReference Include="..\Shared\EpicGames.Core\EpicGames.Core.csproj" PrivateAssets="All"><Private>false</Private></ProjectReference>
+    <ProjectReference Include="..\UnrealBuildTool\UnrealBuildTool.csproj" PrivateAssets="All"><Private>false</Private></ProjectReference>
+    <ProjectReference Include="..\AutomationTool\AutomationUtils\AutomationUtils.Automation.csproj" PrivateAssets="All"><Private>false</Private></ProjectReference>
+  </ItemGroup>
+  <ItemGroup>
+    <PackageReference Include="System.Data.DataSetExtensions" Version="4.5.0" />
+  </ItemGroup>
+</Project>
\ No newline at end of file
diff --git a/Engine/Source/Programs/LiveLinkHub/LiveLinkHub.Target.cs b/Engine/Source/Programs/LiveLinkHub/LiveLinkHub.Target.cs
--- a/Engine/Source/Programs/LiveLinkHub/LiveLinkHub.Target.cs
+++ b/Engine/Source/Programs/LiveLinkHub/LiveLinkHub.Target.cs
@@ -1,13 +1,14 @@
 // Copyright Epic Games, Inc. All Rights Reserved.
 
-using UnrealBuildTool;
-using System.Collections.Generic;
 using EpicGames.Core;
+using UnrealBuildTool;
+
 
 [SupportedPlatforms(UnrealPlatformClass.Desktop)]
 public class LiveLinkHubTarget : TargetRules
 {
 	// Whether this should be built as a monolithic executable.
+	// FIXME?: This shadows an existing UBT switch, but the default is flipped.
 	[CommandLine("-Monolithic")]
 	public bool bMonolithic = false;
 
@@ -36,7 +37,7 @@ public class LiveLinkHubTarget : TargetRules
 		
 		// These plugins are required for running LiveLinkHub. 
 		// They may be a direct dependency or a dependency of one of our plugins.
-		AdditionalPlugins.AddRange(new string[]
+		EnablePlugins.AddRange(new string[]
 		{
 			"LiveLink",
 			"LiveLinkHub",
@@ -48,10 +49,11 @@ public class LiveLinkHubTarget : TargetRules
 			"ProceduralMeshComponent", // Needed by LensComponent
 			"PropertyAccessEditor",
 			"PythonScriptPlugin",
-			"UdpMessaging"
+			"UdpMessaging",
+			"CameraCalibrationCore",
+			"AppleARKitFaceSupport",
 		});
-		
-		
+
 		if (bEnableThirdPartyPlugins)
 		{
 			OptionalPlugins.AddRange(new string[]
@@ -59,7 +61,9 @@ public class LiveLinkHubTarget : TargetRules
 				"LiveLinkViconDataStream",
 				"MocopiLiveLink",
 				"PoseAILiveLink",
-				"Smartsuit"
+				"Smartsuit",
+				//"LiveLinkMVNPlugin",
+				//"OptitrackLiveLink",
 			});
 		}
 
@@ -71,11 +75,6 @@ public class LiveLinkHubTarget : TargetRules
 			});
 		}
 
-		OptionalPlugins.AddRange(new string[]
-		{
-			"AppleARKitFaceSupport"
-		});
-
 		bCompileAgainstCoreUObject = true;
 		bCompileAgainstEngine = true;
 		bCompileAgainstEditor = true;
@@ -103,6 +102,9 @@ public class LiveLinkHubTarget : TargetRules
 
 		bEnableTrace = true;
 
+		// FIXME?: Without this, staging fails on several files in TargetPlatform-related
+		// restricted subdirectories (Engine/Binaries/Win64/{Android,IOS,Linux,LinuxArm64...}).
+		//OptedInModulePlatforms = new UnrealTargetPlatform[] { Target.Platform };
 		OptedInModulePlatforms = new UnrealTargetPlatform[] { UnrealTargetPlatform.Win64, UnrealTargetPlatform.Mac,
 															  UnrealTargetPlatform.Linux, UnrealTargetPlatform.LinuxArm64 };
 	}
diff --git a/Engine/Source/Programs/LiveLinkHub/LiveLinkHubLauncher.Build.cs b/Engine/Source/Programs/LiveLinkHub/LiveLinkHubLauncher.Build.cs
--- a/Engine/Source/Programs/LiveLinkHub/LiveLinkHubLauncher.Build.cs
+++ b/Engine/Source/Programs/LiveLinkHub/LiveLinkHubLauncher.Build.cs
@@ -1,9 +1,13 @@
 // Copyright Epic Games, Inc. All Rights Reserved.
 
-using Microsoft.Extensions.Logging;
+#nullable enable
+
 using System;
+using System.Collections.Generic;
+using System.Linq;
+using EpicGames.Core;
 using UnrealBuildTool;
-using System.IO;
+
 
 // This is named LiveLinkHubLauncher in order to not clash with the LiveLinkHub module which is in the livelink plugin.
 public class LiveLinkHubLauncher : ModuleRules
@@ -18,6 +22,7 @@ public class LiveLinkHubLauncher : ModuleRules
 				"ApplicationCore",
 				"Core",
 				"CoreUObject",
+				"DesktopPlatform",
 				"Engine",
 				"InputCore",
 				"InstallBundleManager",
@@ -74,5 +79,46 @@ public class LiveLinkHubLauncher : ModuleRules
 		}
 
 		ShortName = "LLHub";
+
+		AddRuntimeDependencies();
+	}
+
+	void AddRuntimeDependencies()
+	{
+		DirectoryReference ModuleDirectoryRef = new DirectoryReference(ModuleDirectory);
+		FileReference ConfigLocation = FileReference.Combine(ModuleDirectoryRef, "RuntimeDependencies.ini");
+		ConfigHierarchy Config = new ConfigHierarchy(new List<ConfigFile>{ new ConfigFile(ConfigLocation) });
+
+		if (!Config.TryGetValues("LiveLinkHubLauncher", "RuntimeDependencies", out IReadOnlyList<string>? ConfigDependencies))
+		{
+			string Err = "Error reading runtime dependencies from config";
+			Console.WriteLine(Err);
+			throw new BuildException(Err);
+		}
+
+		FileFilter Filter = new FileFilter();
+		foreach (string Rule in ConfigDependencies)
+		{
+			Filter.AddRule(Rule);
+		}
+
+		// Apply the standard exclusion rules (c.f. DeploymentContext ctor)
+		HashSet<string> RestrictedFolderNames = new HashSet<string>(StringComparer.OrdinalIgnoreCase);
+		RestrictedFolderNames.UnionWith(PlatformExports.GetPlatformFolderNames());
+		RestrictedFolderNames.UnionWith(RestrictedFolder.GetNames());
+		RestrictedFolderNames.ExceptWith(PlatformExports.GetIncludedFolderNames(Target.Platform));
+		foreach (string RestrictedFolderName in RestrictedFolderNames)
+		{
+			string Rule = String.Format(".../{0}/...", RestrictedFolderName);
+			Filter.AddRule(Rule, FileFilterType.Exclude);
+		}
+
+		// Make relative to $(EngineDir) and add to RuntimeDependencies
+		List<FileReference> AbsFiles = Filter.ApplyToDirectory(new DirectoryReference(EngineDirectory), true);
+		IEnumerable<string> RelFiles = AbsFiles.Select(x => x.FullName.Replace(EngineDirectory, "$(EngineDir)"));
+		foreach (string RelFile in RelFiles)
+		{
+			RuntimeDependencies.Add(RelFile, StagedFileType.UFS);
+		}
 	}
 }
diff --git a/Engine/Source/Programs/LiveLinkHub/Private/LiveLinkHubRun.cpp b/Engine/Source/Programs/LiveLinkHub/Private/LiveLinkHubRun.cpp
--- a/Engine/Source/Programs/LiveLinkHub/Private/LiveLinkHubRun.cpp
+++ b/Engine/Source/Programs/LiveLinkHub/Private/LiveLinkHubRun.cpp
@@ -25,6 +25,17 @@ int32 RunLiveLinkHub(const TCHAR* CommandLine)
 {
 	FTaskTagScope Scope(ETaskTag::EGameThread);
 
+#if !(UE_BUILD_SHIPPING)
+	if (FParse::Param(CommandLine, TEXT("WaitForDebugger")))
+	{
+		while (!FPlatformMisc::IsDebuggerPresent())
+		{
+			FPlatformProcess::Sleep(0.1f);
+		}
+		UE_DEBUG_BREAK();
+	}
+#endif
+
 #if WITH_ASSET_LOADING_AUDIT
 	FCoreDelegates::OnSyncLoadPackage.AddLambda([](const FString& PackageName)
 		{
@@ -83,7 +94,7 @@ int32 RunLiveLinkHub(const TCHAR* CommandLine)
 	GIsSilent = true;
 
 	// Start up the main loop, adding some extra command line arguments:
-	const int32 Result = GEngineLoop.PreInit(*FString::Printf(TEXT("%s %s"), CommandLine, TEXT("LiveLinkHubCommandlet -Messaging -DDC=NoShared -NoShaderCompile")));
+	const int32 Result = GEngineLoop.PreInit(*FString::Printf(TEXT("%s %s"), CommandLine, TEXT("-RUN=LiveLinkHubCommandlet -Messaging -DDC=NoShared -NoShaderCompile")));
 
 	if (Result != 0)
 	{
@@ -94,4 +105,4 @@ int32 RunLiveLinkHub(const TCHAR* CommandLine)
 	GEngineLoop.Exit();
 
 	return Result;
-}
\ No newline at end of file
+}
diff --git a/Engine/Source/Programs/LiveLinkHub/Private/LiveLinkHubRunMain.cpp b/Engine/Source/Programs/LiveLinkHub/Private/LiveLinkHubRunMain.cpp
--- a/Engine/Source/Programs/LiveLinkHub/Private/LiveLinkHubRunMain.cpp
+++ b/Engine/Source/Programs/LiveLinkHub/Private/LiveLinkHubRunMain.cpp
@@ -3,4 +3,8 @@
 #include "HAL/Platform.h"
 #include "RequiredProgramMainCPPInclude.h"
 
+#if IS_PROGRAM
 IMPLEMENT_APPLICATION(LiveLinkHubLauncher, "LiveLinkHub");
+#else
+FEngineLoop GEngineLoop;
+#endif
diff --git a/Engine/Source/Programs/LiveLinkHub/RuntimeDependencies.ini b/Engine/Source/Programs/LiveLinkHub/RuntimeDependencies.ini
new file mode 100644
--- /dev/null
+++ b/Engine/Source/Programs/LiveLinkHub/RuntimeDependencies.ini
@@ -0,0 +1,154 @@
+[LiveLinkHubLauncher]
++RuntimeDependencies="/Content/Animation/..."
++RuntimeDependencies="/Content/BasicShapes/..."
++RuntimeDependencies="/Content/Editor/Slate/..."
+; Prefix - denotes exclusion
++RuntimeDependencies="-/Content/Editor/Slate/Automation/..."
++RuntimeDependencies="-/Content/Editor/Slate/Testing/..."
++RuntimeDependencies="/Content/EditorMeshes/Camera/..."
++RuntimeDependencies="/Content/EngineDamageTypes/..."
++RuntimeDependencies="/Content/EngineFonts/..."
++RuntimeDependencies="/Content/Slate/..."
+
+; Individual assets
++RuntimeDependencies="/Content/ArtTools/RenderToTexture/Meshes/S_1_Unit_Plane.uasset"
++RuntimeDependencies="/Content/EditorLandscapeResources/DataLayer.uasset"
++RuntimeDependencies="/Content/EditorLandscapeResources/DefaultAlphaTexture.uasset"
++RuntimeDependencies="/Content/EditorLandscapeResources/SplineEditorMesh.uasset"
++RuntimeDependencies="/Content/EditorLandscapeResources/SplineEditorMeshMat.uasset"
++RuntimeDependencies="/Content/EditorMaterials/Camera/CineMat.uasset"
++RuntimeDependencies="/Content/EditorMaterials/Camera/MI_CineMat_Rig.uasset"
++RuntimeDependencies="/Content/EditorMaterials/ParticleSystems/PSysThumbnail_NoImage.uasset"
++RuntimeDependencies="/Content/EditorMaterials/ParticleSystems/PSysThumbnail_OOD.uasset"
++RuntimeDependencies="/Content/EditorMaterials/WorldPartition/WorldPartitionSpatialHashGridPreviewMaterial.uasset"
++RuntimeDependencies="/Content/EditorResources/S_Actor.uasset"
++RuntimeDependencies="/Content/EditorResources/S_BoxReflectionCapture.uasset"
++RuntimeDependencies="/Content/EditorResources/S_Solver.uasset"
++RuntimeDependencies="/Content/EditorResources/S_Terrain.uasset"
++RuntimeDependencies="/Content/EditorResources/S_Trigger.uasset"
++RuntimeDependencies="/Content/EditorResources/S_TriggerBox.uasset"
++RuntimeDependencies="/Content/EditorResources/S_TriggerCapsule.uasset"
++RuntimeDependencies="/Content/EditorResources/S_TriggerSphere.uasset"
++RuntimeDependencies="/Content/EngineDebugMaterials/DebugMeshMaterial.uasset"
++RuntimeDependencies="/Content/EngineDebugMaterials/M_SimpleUnlitTranslucent.uasset"
++RuntimeDependencies="/Content/EngineDebugMaterials/M_VolumeRenderSphereTracePP.uasset"
++RuntimeDependencies="/Content/EngineDebugMaterials/VolumeToRender.uasset"
++RuntimeDependencies="/Content/EngineMaterials/BaseFlattenDiffuseMap.uasset"
++RuntimeDependencies="/Content/EngineMaterials/BaseFlattenDiffuseMap_VT.uasset"
++RuntimeDependencies="/Content/EngineMaterials/BaseFlattenEmissiveMap.uasset"
++RuntimeDependencies="/Content/EngineMaterials/BaseFlattenEmissiveMap_VT.uasset"
++RuntimeDependencies="/Content/EngineMaterials/BaseFlattenGrayscaleMap.uasset"
++RuntimeDependencies="/Content/EngineMaterials/BaseFlattenGrayscaleMap_VT.uasset"
++RuntimeDependencies="/Content/EngineMaterials/BaseFlattenLinearColor.uasset"
++RuntimeDependencies="/Content/EngineMaterials/BaseFlattenLinearColor_VT.uasset"
++RuntimeDependencies="/Content/EngineMaterials/BaseFlattenNormalMap.uasset"
++RuntimeDependencies="/Content/EngineMaterials/BaseFlattenNormalMap_VT.uasset"
++RuntimeDependencies="/Content/EngineMaterials/Black_1x1_EXR_Texture_VT.uasset"
++RuntimeDependencies="/Content/EngineMaterials/DefaultBokeh.uasset"
++RuntimeDependencies="/Content/EngineMaterials/DefaultDeferredDecalMaterial.uasset"
++RuntimeDependencies="/Content/EngineMaterials/DefaultDeferredDecalMaterial.uasset"
++RuntimeDependencies="/Content/EngineMaterials/DefaultDestructiblePhysicalMaterial.uasset"
++RuntimeDependencies="/Content/EngineMaterials/DefaultDiffuse.uasset"
++RuntimeDependencies="/Content/EngineMaterials/DefaultLightFunctionMaterial.uasset"
++RuntimeDependencies="/Content/EngineMaterials/DefaultLightFunctionMaterial.uasset"
++RuntimeDependencies="/Content/EngineMaterials/DefaultMaterial.uasset"
++RuntimeDependencies="/Content/EngineMaterials/DefaultPhysicalMaterial.uasset"
++RuntimeDependencies="/Content/EngineMaterials/DefaultPostProcessMaterial.uasset"
++RuntimeDependencies="/Content/EngineMaterials/DefaultPostProcessMaterial.uasset"
++RuntimeDependencies="/Content/EngineMaterials/DefaultTextMaterialOpaque.uasset"
++RuntimeDependencies="/Content/EngineMaterials/DefaultWhiteGrid.uasset"
++RuntimeDependencies="/Content/EngineMaterials/DefaultWhiteGrid_Low.uasset"
++RuntimeDependencies="/Content/EngineMaterials/EmissiveMeshMaterial.uasset"
++RuntimeDependencies="/Content/EngineMaterials/FlattenMaterial.uasset"
++RuntimeDependencies="/Content/EngineMaterials/FlattenMaterial_VT.uasset"
++RuntimeDependencies="/Content/EngineMaterials/FlattenMaterial_WS_Normal_VT.uasset"
++RuntimeDependencies="/Content/EngineMaterials/GizmoMaterial.uasset"
++RuntimeDependencies="/Content/EngineMaterials/Good64x64TilingNoiseHighFreq.uasset"
++RuntimeDependencies="/Content/EngineMaterials/InvalidLightmapSettings.uasset"
++RuntimeDependencies="/Content/EngineMaterials/M_InvalidLightmapSettings.uasset"
++RuntimeDependencies="/Content/EngineMaterials/MiniFont.uasset"
++RuntimeDependencies="/Content/EngineMaterials/NaniteHiddenSectionMaterial.uasset"
++RuntimeDependencies="/Content/EngineMaterials/PhAT_JointLimitMaterial.uasset"
++RuntimeDependencies="/Content/EngineMaterials/PreintegratedSkinBRDF.uasset"
++RuntimeDependencies="/Content/EngineMaterials/RemoveSurfaceMaterial.uasset"
++RuntimeDependencies="/Content/EngineMaterials/T_Default_Material_Grid_M.uasset"
++RuntimeDependencies="/Content/EngineMaterials/T_Default_Material_Grid_N.uasset"
++RuntimeDependencies="/Content/EngineMaterials/WeightMapArrayPlaceholderTexture.uasset"
++RuntimeDependencies="/Content/EngineMaterials/WeightMapPlaceholderTexture.uasset"
++RuntimeDependencies="/Content/EngineMaterials/Widget3DPassThrough.uasset"
++RuntimeDependencies="/Content/EngineMaterials/Widget3DPassThrough_Masked.uasset"
++RuntimeDependencies="/Content/EngineMaterials/Widget3DPassThrough_Masked_OneSided.uasset"
++RuntimeDependencies="/Content/EngineMaterials/Widget3DPassThrough_Opaque.uasset"
++RuntimeDependencies="/Content/EngineMaterials/Widget3DPassThrough_Opaque_OneSided.uasset"
++RuntimeDependencies="/Content/EngineMaterials/Widget3DPassThrough_Translucent.uasset"
++RuntimeDependencies="/Content/EngineMaterials/Widget3DPassThrough_Translucent_OneSided.uasset"
++RuntimeDependencies="/Content/EngineMaterials/WorldGridMaterial.uasset"
++RuntimeDependencies="/Content/EngineMaterials/WorldGridMaterial.uasset"
++RuntimeDependencies="/Content/EngineMeshes/Sphere.uasset"
++RuntimeDependencies="/Content/EngineResources/Black.uasset"
++RuntimeDependencies="/Content/EngineResources/DefaultTexture.uasset"
++RuntimeDependencies="/Content/EngineResources/DefaultTextureCube.uasset"
++RuntimeDependencies="/Content/EngineResources/DefaultVolumeTexture.uasset"
++RuntimeDependencies="/Content/EngineResources/DefaultVolumeTexture2D.uasset"
++RuntimeDependencies="/Content/EngineResources/GradientTexture0.uasset"
++RuntimeDependencies="/Content/EngineResources/WhiteSquareTexture.uasset"
++RuntimeDependencies="/Content/EngineSky/VolumetricClouds/CloudGradientTexture.uasset"
++RuntimeDependencies="/Content/EngineSky/VolumetricClouds/CloudWeatherTexture.uasset"
++RuntimeDependencies="/Content/EngineSky/VolumetricClouds/T_NoiseErosion.uasset"
++RuntimeDependencies="/Content/EngineSky/VolumetricClouds/T_NoiseShape64.uasset"
++RuntimeDependencies="/Content/EngineSky/VolumetricClouds/T_VolumeNoiseErosion32.uasset"
++RuntimeDependencies="/Content/EngineSky/VolumetricClouds/T_VolumeNoiseShape64.uasset"
++RuntimeDependencies="/Content/EngineSky/VolumetricClouds/m_SimpleVolumetricCloud.uasset"
++RuntimeDependencies="/Content/EngineSky/VolumetricClouds/m_SimpleVolumetricCloud_Inst.uasset"
++RuntimeDependencies="/Content/Functions/Engine_MaterialFunctions01/Opacity/CameraDepthFade.uasset"
++RuntimeDependencies="/Content/Functions/Engine_MaterialFunctions01/Shading/ConvertFromDiffSpec.uasset"
++RuntimeDependencies="/Content/Functions/Engine_MaterialFunctions01/Shading/PowerToRoughness.uasset"
++RuntimeDependencies="/Content/Functions/Engine_MaterialFunctions02/Math/Transform3x3Matrix.uasset"
++RuntimeDependencies="/Content/Functions/Engine_MaterialFunctions02/PivotPainter2/Black_1x1_EXR_Texture.uasset"
++RuntimeDependencies="/Content/Functions/Engine_MaterialFunctions02/SplitComponents.uasset"
++RuntimeDependencies="/Content/Functions/Engine_MaterialFunctions02/Utility/BreakOutFloat2Components.uasset"
++RuntimeDependencies="/Content/Functions/Engine_MaterialFunctions02/Utility/BreakOutFloat3Components.uasset"
++RuntimeDependencies="/Content/Functions/Engine_MaterialFunctions02/Utility/CameraDirectionVector.uasset"
++RuntimeDependencies="/Content/Functions/Engine_MaterialFunctions02/Utility/MakeFloat2.uasset"
++RuntimeDependencies="/Content/Functions/Engine_MaterialFunctions02/Utility/MakeFloat3.uasset"
++RuntimeDependencies="/Content/Functions/Engine_MaterialFunctions02/Utility/MakeFloat4.uasset"
++RuntimeDependencies="/Content/Functions/Engine_MaterialFunctions02/WorldPositionOffset/CameraOffset.uasset"
++RuntimeDependencies="/Content/VREditor/LaserPointer/VR_LaserPower_01.uasset"
++RuntimeDependencies="/Content/VREditor/SnapGrid/SnapGridMaterial.uasset"
++RuntimeDependencies="/Content/VREditor/SnapGrid/SnapGridPlaneMesh.uasset"
++RuntimeDependencies="/Content/VREditor/Sounds/UI/Drop_Item_From_ContentBrowser.uasset"
++RuntimeDependencies="/Content/VREditor/Sounds/UI/Gizmo_Handle_Clicked.uasset"
++RuntimeDependencies="/Content/VREditor/Sounds/UI/Object_PickUp.uasset"
++RuntimeDependencies="/Content/VREditor/Sounds/UI/Object_Snaps_To_Another_Actor.uasset"
++RuntimeDependencies="/Content/VREditor/Sounds/UI/Object_Snaps_To_Grid.uasset"
++RuntimeDependencies="/Content/VREditor/Sounds/UI/Scaling_Down.uasset"
++RuntimeDependencies="/Content/VREditor/Sounds/UI/Scaling_Up.uasset"
++RuntimeDependencies="/Content/VREditor/Sounds/UI/Selection_Changes.uasset"
++RuntimeDependencies="/Content/VREditor/TransformGizmo/BoundingBoxCorner.uasset"
++RuntimeDependencies="/Content/VREditor/TransformGizmo/BoundingBoxEdge.uasset"
++RuntimeDependencies="/Content/VREditor/TransformGizmo/Main.uasset"
++RuntimeDependencies="/Content/VREditor/TransformGizmo/PlaneTranslationHandle.uasset"
++RuntimeDependencies="/Content/VREditor/TransformGizmo/RotationHandleFull.uasset"
++RuntimeDependencies="/Content/VREditor/TransformGizmo/RotationHandleIndicator.uasset"
++RuntimeDependencies="/Content/VREditor/TransformGizmo/RotationHandleQuarter.uasset"
++RuntimeDependencies="/Content/VREditor/TransformGizmo/SM_Sequencer_Node.uasset"
++RuntimeDependencies="/Content/VREditor/TransformGizmo/StartRotationHandleIndicator.uasset"
++RuntimeDependencies="/Content/VREditor/TransformGizmo/TransformGizmoFreeRotation.uasset"
++RuntimeDependencies="/Content/VREditor/TransformGizmo/TransformGizmoMaterial.uasset"
++RuntimeDependencies="/Content/VREditor/TransformGizmo/TransformGizmoMaterial_Inst.uasset"
++RuntimeDependencies="/Content/VREditor/TransformGizmo/TranslateArrowHandle.uasset"
++RuntimeDependencies="/Content/VREditor/TransformGizmo/TranslucentTransformGizmoMaterial.uasset"
++RuntimeDependencies="/Content/VREditor/TransformGizmo/TranslucentTransformGizmoMaterial_Inst.uasset"
++RuntimeDependencies="/Content/VREditor/TransformGizmo/UniformScaleHandle.uasset"
++RuntimeDependencies="/Content/VREditor/TransformGizmo/Xray.uasset"
++RuntimeDependencies="/Content/VREditor/ViewportInteractionAssetContainerData.uasset"
+
+; Plugin content (only needed for program BuildCookRun; MakeCookedEditor includes by default)
+; FIXME: ACLPlugin still generates warnings at runtime
++RuntimeDependencies="/Plugins/Animation/ACLPlugin/Content/..."
++RuntimeDependencies="/Plugins/Animation/LiveLink/Content/..."
++RuntimeDependencies="/Plugins/MovieScene/LevelSequenceEditor/Content/..."
++RuntimeDependencies="/Plugins/Runtime/AR/ARUtilities/Content/..."
++RuntimeDependencies="/Plugins/Runtime/AR/AppleAR/AppleARKit/Content/..."
++RuntimeDependencies="/Plugins/VirtualProduction/CameraCalibrationCore/Content/..."
++RuntimeDependencies="/Plugins/VirtualProduction/Takes/Content/..."
diff --git a/Engine/Source/Programs/LiveLinkHubEditor/Config/DefaultGame.ini b/Engine/Source/Programs/LiveLinkHubEditor/Config/DefaultGame.ini
new file mode 100644
--- /dev/null
+++ b/Engine/Source/Programs/LiveLinkHubEditor/Config/DefaultGame.ini
@@ -0,0 +1,9 @@
+[CookedEditorSettings_CookedEditor]
+CookedEditorTargetName=LiveLinkHubEditor
+bStageShaderDirs=false
+bStagePlatformBuildDirs=false
+bStageExtrasDirs=false
+bStagePlatformDirs=false
+bStageUAT=false
+bStagePython=false
+bIsForExternalDistribution=true
diff --git a/Engine/Source/Programs/LiveLinkHubEditor/LiveLinkHubEditor.uproject b/Engine/Source/Programs/LiveLinkHubEditor/LiveLinkHubEditor.uproject
new file mode 100644
--- /dev/null
+++ b/Engine/Source/Programs/LiveLinkHubEditor/LiveLinkHubEditor.uproject
@@ -0,0 +1,11 @@
+{
+	"FileVersion": 3,
+	"EngineAssociation": "",
+	"Category": "",
+	"Description": "",
+	"TargetPlatforms": [
+		"Linux",
+		"Mac",
+		"Windows"
+	]
+}
diff --git a/Engine/Source/Programs/LiveLinkHubEditor/Source/LiveLinkHubEditor.Target.cs b/Engine/Source/Programs/LiveLinkHubEditor/Source/LiveLinkHubEditor.Target.cs
new file mode 100644
--- /dev/null
+++ b/Engine/Source/Programs/LiveLinkHubEditor/Source/LiveLinkHubEditor.Target.cs
@@ -0,0 +1,75 @@
+// Copyright Epic Games, Inc. All Rights Reserved.
+
+using EpicGames.Core;
+using UnrealBuildTool;
+
+
+[SupportedPlatforms(UnrealPlatformClass.Desktop)]
+public class LiveLinkHubEditorTarget : TargetRules
+{
+	// Whether to enable building Capture Manager plugin.
+	[CommandLine("-EnableCaptureManagerPlugin=")]
+	public bool bEnableCaptureManagerPlugin = false;
+
+	public LiveLinkHubEditorTarget(TargetInfo Target) : base(Target)
+	{
+		Type = TargetType.Editor;
+		bExplicitTargetForType = true;
+
+		SolutionDirectory = "Programs/LiveLink";
+		LaunchModuleName = "LiveLinkHubLauncher";
+		IncludeOrderVersion = EngineIncludeOrderVersion.Latest;
+
+		// These plugins are required for running LiveLinkHub. 
+		// They may be a direct dependency or a dependency of another plugin.
+		EnablePlugins.AddRange(new string[]
+		{
+			"LiveLink",
+			"LiveLinkHub",
+			"LiveLinkCamera",
+			"LiveLinkLens",
+			"LensComponent",
+			"LiveLinkInputDevice",
+			"ContentBrowserAssetDataSource",
+			"ProceduralMeshComponent",
+			"PropertyAccessEditor",
+			"PythonScriptPlugin",
+			"StructUtils",
+			"UdpMessaging",
+			"CameraCalibrationCore",
+			"AppleARKitFaceSupport",
+		});
+
+		if (bEnableCaptureManagerPlugin)
+		{
+			OptionalPlugins.AddRange(new string[]
+			{
+				"CaptureManager"
+			});
+		}
+
+		// FIXME?: Without this, staging fails on several files in TargetPlatform-related
+		// restricted subdirectories (Engine/Binaries/Win64/{Android,IOS,Linux,LinuxArm64...}).
+		OptedInModulePlatforms = new UnrealTargetPlatform[] { Target.Platform };
+
+		bAllowEnginePluginsEnabledByDefault = false;
+		bBuildAdditionalConsoleApp = false;
+
+		// Based loosely on the VCProject.cs logic to construct NMakePath.
+		string BaseExeName = "LiveLinkHubEditor";
+		OutputFile = "Binaries/" + Platform.ToString() + "/" + BaseExeName;
+		if (Configuration != UndecoratedConfiguration)
+		{
+			OutputFile += "-" + Platform.ToString() + "-" + Configuration.ToString();
+		}
+		if (Platform == UnrealTargetPlatform.Win64)
+		{
+			OutputFile += ".exe";
+		}
+
+		if (bGenerateProjectFiles)
+		{
+			Type = TargetType.Program;
+		}
+	}
+}
