LiveLink Hub (Editor target): Don't use a commandlet or standalone Slate renderer, run the full editor initialization and bring up the RHI instead.

This is necessary to support planned future development that requires a proper RHI. It also obviates the need for some previous workarounds related to operating as a commandlet.

Also contains some fixes related to running Python automation.

diff --git a/Engine/Plugins/Animation/LiveLinkHub/Source/LiveLinkHub/Private/LiveLinkHub.cpp b/Engine/Plugins/Animation/LiveLinkHub/Source/LiveLinkHub/Private/LiveLinkHub.cpp
--- a/Engine/Plugins/Animation/LiveLinkHub/Source/LiveLinkHub/Private/LiveLinkHub.cpp
+++ b/Engine/Plugins/Animation/LiveLinkHub/Source/LiveLinkHub/Private/LiveLinkHub.cpp
@@ -28,16 +28,22 @@
 
 #define LOCTEXT_NAMESPACE "LiveLinkHub"
 
+
+void FLiveLinkHub::Preinitialize()
+{
+	// We must register the livelink client first since we might rely on the modular feature to initialize the controllers/managers.
+	LiveLinkHubClient = MakeShared<FLiveLinkHubClient>(AsShared());
+	IModularFeatures::Get().RegisterModularFeature(ILiveLinkClient::ModularFeatureName, LiveLinkHubClient.Get());
+}
+
 void FLiveLinkHub::Initialize()
 {
 	TRACE_CPUPROFILER_EVENT_SCOPE(FLiveLinkHub::Initialize);
 
+#if IS_PROGRAM
 	// Re-enable this since we've disabled it to avoid the creation of the console window.
 	GIsSilent = false;
-
-	// We must register the livelink client first since we might rely on the modular feature to initialize the controllers/managers.
-	LiveLinkHubClient = MakeShared<FLiveLinkHubClient>(AsShared());
-	IModularFeatures::Get().RegisterModularFeature(ILiveLinkClient::ModularFeatureName, LiveLinkHubClient.Get());
+#endif
 
 	SessionManager = MakeShared<FLiveLinkHubSessionManager>();
 	LiveLinkProvider = MakeShared<FLiveLinkHubProvider>(SessionManager.ToSharedRef());
diff --git a/Engine/Plugins/Animation/LiveLinkHub/Source/LiveLinkHub/Private/LiveLinkHub.h b/Engine/Plugins/Animation/LiveLinkHub/Source/LiveLinkHub/Private/LiveLinkHub.h
--- a/Engine/Plugins/Animation/LiveLinkHub/Source/LiveLinkHub/Private/LiveLinkHub.h
+++ b/Engine/Plugins/Animation/LiveLinkHub/Source/LiveLinkHub/Private/LiveLinkHub.h
@@ -49,6 +49,8 @@ public:
 	//~ End ILiveLinkHub interface
 
 public:
+	/** First phase initialization; must precede FEngineLoop::Init(). */
+	void Preinitialize();
 	/** Launch the slate application and initialize its components. */
 	void Initialize();
 	/** Tick the hub. */
diff --git a/Engine/Plugins/Animation/LiveLinkHub/Source/LiveLinkHub/Private/LiveLinkHubModule.cpp b/Engine/Plugins/Animation/LiveLinkHub/Source/LiveLinkHub/Private/LiveLinkHubModule.cpp
--- a/Engine/Plugins/Animation/LiveLinkHub/Source/LiveLinkHub/Private/LiveLinkHubModule.cpp
+++ b/Engine/Plugins/Animation/LiveLinkHub/Source/LiveLinkHub/Private/LiveLinkHubModule.cpp
@@ -20,16 +20,27 @@
 
 #define LOCTEXT_NAMESPACE "LiveLinkHubModule"
 
+void FLiveLinkHubModule::PreinitializeLiveLinkHub()
+{
+	check(!LiveLinkHub);
+	LiveLinkHub = MakeShared<FLiveLinkHub>();
+	LiveLinkHub->Preinitialize();
+}
+
 void FLiveLinkHubModule::StartLiveLinkHub()
 {
 	TRACE_CPUPROFILER_EVENT_SCOPE(StartLiveLinkHub);
-	// Nothing will get executed after this, so put everything before.
-	LiveLinkHub = MakeShared<FLiveLinkHub>();
+	checkf(LiveLinkHub, TEXT("Ensure PreinitializeLiveLinkHub was called first"));
+
 	LiveLinkHub->Initialize();
 
+#if IS_PROGRAM
 	LiveLinkHubLoop(LiveLinkHub);
+#endif
+}
 
-	// If we got to this point, the app is shutdown so we should release our references.
+void FLiveLinkHubModule::ShutdownLiveLinkHub()
+{
 	LiveLinkHub.Reset();
 }
 
diff --git a/Engine/Plugins/Animation/LiveLinkHub/Source/LiveLinkHub/Private/LiveLinkHubModule.h b/Engine/Plugins/Animation/LiveLinkHub/Source/LiveLinkHub/Private/LiveLinkHubModule.h
--- a/Engine/Plugins/Animation/LiveLinkHub/Source/LiveLinkHub/Private/LiveLinkHubModule.h
+++ b/Engine/Plugins/Animation/LiveLinkHub/Source/LiveLinkHub/Private/LiveLinkHubModule.h
@@ -19,7 +19,9 @@ class FLiveLinkHubModule : public ILiveLinkHubModule
 {
 public:
 	//~ Begin ILiveLinkHubModule interface
+	virtual void PreinitializeLiveLinkHub() override;
 	virtual void StartLiveLinkHub() override;
+	virtual void ShutdownLiveLinkHub() override;
 	virtual void StartupModule() override;
 	virtual void ShutdownModule() override;
 	//~ End ILiveLinkHubModule interface
diff --git a/Engine/Plugins/Animation/LiveLinkHub/Source/LiveLinkHub/Private/UI/Window/LiveLinkHubWindowController.cpp b/Engine/Plugins/Animation/LiveLinkHub/Source/LiveLinkHub/Private/UI/Window/LiveLinkHubWindowController.cpp
--- a/Engine/Plugins/Animation/LiveLinkHub/Source/LiveLinkHub/Private/UI/Window/LiveLinkHubWindowController.cpp
+++ b/Engine/Plugins/Animation/LiveLinkHub/Source/LiveLinkHub/Private/UI/Window/LiveLinkHubWindowController.cpp
@@ -25,7 +25,15 @@ FLiveLinkHubWindowController::FLiveLinkHubWindowController(const FLiveLinkHubWin
 FLiveLinkHubWindowController::~FLiveLinkHubWindowController()
 {
 	FGlobalTabmanager::Get()->SaveAllVisualState();
+
+	if (RootWindow)
+	{
+		RootWindow->SetOnWindowClosed(nullptr);
+	}
+
+#if IS_PROGRAM
 	FSlateApplication::Shutdown();
+#endif
 }
 
 TSharedRef<SWindow> FLiveLinkHubWindowController::CreateWindow()
@@ -96,7 +104,9 @@ void FLiveLinkHubWindowController::RestoreLayout()
 
 TSharedPtr<FModalWindowManager> FLiveLinkHubWindowController::InitializeSlateApplication()
 {
-	TRACE_CPUPROFILER_EVENT_SCOPE(FLiveLinkHubWindowController::InitializeAsStandaloneApplication);
+	TRACE_CPUPROFILER_EVENT_SCOPE(FLiveLinkHubWindowController::InitializeSlateApplication);
+
+#if IS_PROGRAM
 	FSlateApplication::InitializeAsStandaloneApplication(GetStandardStandaloneRenderer());
 
 	// @hack This call will silently fail since we're running a commandlet, so pretend like we aren't one for it.
@@ -105,6 +115,7 @@ TSharedPtr<FModalWindowManager> FLiveLinkHubWindowController::InitializeSlateApp
 	PRIVATE_GIsRunningCommandlet = false;
 	FSlateApplication::InitHighDPI(true);
 	PRIVATE_GIsRunningCommandlet = bIsRunningCommandlet;
+#endif
 
 	const FText ApplicationTitle = LOCTEXT("AppTitle", "LiveLink Hub");
 	FGlobalTabmanager::Get()->SetApplicationTitle(ApplicationTitle);
@@ -128,6 +139,10 @@ void FLiveLinkHubWindowController::OnWindowClosed(const TSharedRef<SWindow>& Win
 {
 	SaveLayout();
 	RootWindow.Reset();
+
+#if !IS_PROGRAM
+	RequestEngineExit(TEXT("FLiveLinkHubWindowController::OnWindowClosed"));
+#endif
 }
 
 void FLiveLinkHubWindowController::SaveLayout() const
diff --git a/Engine/Plugins/Animation/LiveLinkHub/Source/LiveLinkHub/Public/ILiveLinkHubModule.h b/Engine/Plugins/Animation/LiveLinkHub/Source/LiveLinkHub/Public/ILiveLinkHubModule.h
--- a/Engine/Plugins/Animation/LiveLinkHub/Source/LiveLinkHub/Public/ILiveLinkHubModule.h
+++ b/Engine/Plugins/Animation/LiveLinkHub/Source/LiveLinkHub/Public/ILiveLinkHubModule.h
@@ -7,7 +7,11 @@
 class LIVELINKHUB_API ILiveLinkHubModule : public IModuleInterface
 {
 public:
+	virtual void PreinitializeLiveLinkHub() = 0;
+
 	/** Launch the slate application hosting the live link hub. */
 	virtual void StartLiveLinkHub() = 0;
+
+	virtual void ShutdownLiveLinkHub() = 0;
 };
 
diff --git a/Engine/Plugins/Animation/LiveLinkHub/Source/LiveLinkHubEditor/Private/SLiveLinkHubEditorStatusBar.cpp b/Engine/Plugins/Animation/LiveLinkHub/Source/LiveLinkHubEditor/Private/SLiveLinkHubEditorStatusBar.cpp
--- a/Engine/Plugins/Animation/LiveLinkHub/Source/LiveLinkHubEditor/Private/SLiveLinkHubEditorStatusBar.cpp
+++ b/Engine/Plugins/Animation/LiveLinkHub/Source/LiveLinkHubEditor/Private/SLiveLinkHubEditorStatusBar.cpp
@@ -33,6 +33,16 @@ void SLiveLinkHubEditorStatusBar::Construct(const FArguments& InArgs)
 	LiveLinkClient->OnLiveLinkSourcesChanged().AddRaw(this, &SLiveLinkHubEditorStatusBar::RefreshSubjects);
 	LiveLinkClient->OnLiveLinkSubjectsChanged().AddRaw(this, &SLiveLinkHubEditorStatusBar::RefreshSubjects);
 
+	IModularFeatures::Get().OnModularFeatureUnregistered().AddSPLambda(this,
+		[this](const FName& InFeatureName, IModularFeature* InFeature)
+		{
+			if (InFeature == this->LiveLinkClient)
+			{
+				this->LiveLinkClient = nullptr;
+			}
+		}
+	);
+
 	constexpr bool bLoop = true;
 	if (GEditor && GEditor->IsTimerManagerValid())
 	{
@@ -112,7 +122,7 @@ SLiveLinkHubEditorStatusBar::~SLiveLinkHubEditorStatusBar()
 		HubMessagingModule->OnConnectionEstablished().RemoveAll(this);
 	}
 
-	IModularFeatures::Get().OnModularFeatureRegistered().RemoveAll(this);
+	IModularFeatures::Get().OnModularFeatureUnregistered().RemoveAll(this);
 
 	if (GEditor && GEditor->IsTimerManagerValid())
 	{
@@ -177,7 +187,7 @@ void SLiveLinkHubEditorStatusBar::CheckHubConnection()
 		return;
 	}
 
-	if (LiveLinkClient->IsSourceStillValid(HubSourceId))
+	if (LiveLinkClient && LiveLinkClient->IsSourceStillValid(HubSourceId))
 	{
 		ConnectionState = EHubConnectionState::Valid;
 		return;
diff --git a/Engine/Source/Programs/AutomationTool/LiveLinkHub/LiveLinkHub.Automation.cs b/Engine/Source/Programs/AutomationTool/LiveLinkHub/LiveLinkHub.Automation.cs
--- a/Engine/Source/Programs/AutomationTool/LiveLinkHub/LiveLinkHub.Automation.cs
+++ b/Engine/Source/Programs/AutomationTool/LiveLinkHub/LiveLinkHub.Automation.cs
@@ -43,6 +43,16 @@ namespace AutomationTool
 			Context.Apply(SC);
 
 			string PlatName = SC.StageTargetPlatform.PlatformType.ToString();
+			string ExeExtension = Platform.GetExeExtension(SC.StageTargetPlatform.PlatformType);
+
+			// Copy Zen binaries
+			SC.FilesToStage.NonUFSFiles.Add(
+				new StagedFileReference($"Engine/Binaries/{PlatName}/zen{ExeExtension}"),
+				new FileReference($"Engine/Binaries/{PlatName}/zen{ExeExtension}"));
+
+			SC.FilesToStage.NonUFSFiles.Add(
+				new StagedFileReference($"Engine/Binaries/{PlatName}/zenserver{ExeExtension}"),
+				new FileReference($"Engine/Binaries/{PlatName}/zenserver{ExeExtension}"));
 
 			// Copy .target receipt to project and engine bin
 			SC.FilesToStage.NonUFSFiles.Add(
diff --git a/Engine/Source/Programs/LiveLinkHub/LiveLinkHubLauncher.Build.cs b/Engine/Source/Programs/LiveLinkHub/LiveLinkHubLauncher.Build.cs
--- a/Engine/Source/Programs/LiveLinkHub/LiveLinkHubLauncher.Build.cs
+++ b/Engine/Source/Programs/LiveLinkHub/LiveLinkHubLauncher.Build.cs
@@ -20,6 +20,7 @@ public class LiveLinkHubLauncher : ModuleRules
 		PrivateDependencyModuleNames.AddRange(
 			new string[] {
 				"ApplicationCore",
+				"AssetRegistry",
 				"Core",
 				"CoreUObject",
 				"DesktopPlatform",
@@ -30,6 +31,8 @@ public class LiveLinkHubLauncher : ModuleRules
 				"Messaging",
 				"MoviePlayer",
 				"MoviePlayerProxy",
+				"ProfileVisualizer",
+				"ProfilerService",
 				"Projects",
 				"PreLoadScreen",
 				"PIEPreviewDeviceProfileSelector",
@@ -52,8 +55,6 @@ public class LiveLinkHubLauncher : ModuleRules
 				"DerivedDataCache",
 				"HeadMountedDisplay", 
 				"MRMesh", 
-				"ProfileVisualizer", 
-				"ProfilerService", 
 				"SlateRHIRenderer", 
 				"SlateNullRenderer",
 			}
diff --git a/Engine/Source/Programs/LiveLinkHub/Private/LiveLinkHubCommandlet.h b/Engine/Source/Programs/LiveLinkHub/Private/LiveLinkHubCommandlet.h
--- a/Engine/Source/Programs/LiveLinkHub/Private/LiveLinkHubCommandlet.h
+++ b/Engine/Source/Programs/LiveLinkHub/Private/LiveLinkHubCommandlet.h
@@ -38,7 +38,11 @@ public:
 		FModuleManager::Get().LoadModule("ProjectSettingsViewer");
 
 		FModuleManager::Get().LoadModuleChecked<ILiveLinkHubModule>("LiveLinkEditor");
-		FModuleManager::Get().LoadModuleChecked<ILiveLinkHubModule>("LiveLinkHub").StartLiveLinkHub();
+
+		ILiveLinkHubModule& HubModule = FModuleManager::Get().LoadModuleChecked<ILiveLinkHubModule>("LiveLinkHub");
+		HubModule.PreinitializeLiveLinkHub();
+		HubModule.StartLiveLinkHub();
+		HubModule.ShutdownLiveLinkHub();
 
 		return 0;
 	}
diff --git a/Engine/Source/Programs/LiveLinkHub/Private/LiveLinkHubRun.cpp b/Engine/Source/Programs/LiveLinkHub/Private/LiveLinkHubRun.cpp
--- a/Engine/Source/Programs/LiveLinkHub/Private/LiveLinkHubRun.cpp
+++ b/Engine/Source/Programs/LiveLinkHub/Private/LiveLinkHubRun.cpp
@@ -8,25 +8,49 @@
 #include "HAL/FileManager.h"
 #include "HAL/PlatformApplicationMisc.h"
 #include "HAL/PlatformFileManager.h"
+#include "HAL/PlatformSplash.h"
+#include "ILiveLinkHubModule.h"
 #include "Interfaces/IPluginManager.h"
 #include "LaunchEngineLoop.h"
 #include "Misc/CommandLine.h"
+#include "Misc/Paths.h"
+#include "Modules/ModuleManager.h"
 
 #ifndef WITH_ASSET_LOADING_AUDIT
 #define WITH_ASSET_LOADING_AUDIT 0
 #endif
 
-#ifndef COOKED_EDITOR_DISTRIBUTION
-#define COOKED_EDITOR_DISTRIBUTION 0
+#if WITH_ASSET_LOADING_AUDIT
+#	include "Misc/CoreDelegates.h"
 #endif
 
+
 DEFINE_LOG_CATEGORY_STATIC(LogLiveLinkHub, Log, All);
 
+
+// Opt in to new D3D12 redist and tell the loader where to search for D3D12Core.dll.
+// The D3D loader looks for these symbol exports in the .exe module.
+// We only support this on x64 Windows Desktop platforms. Other platforms or non-redist-aware 
+// versions of Windows will transparently load default OS-provided D3D12 library.
+#define USE_D3D12_REDIST (PLATFORM_DESKTOP && PLATFORM_CPU_X86_FAMILY && PLATFORM_64BITS && 1)
+#if USE_D3D12_REDIST
+extern "C" { _declspec(dllexport) extern const UINT D3D12SDKVersion = 614; } // D3D12_SDK_VERSION
+extern "C" { _declspec(dllexport) extern const char* D3D12SDKPath = ".\\D3D12\\"; }
+#endif // USE_D3D12_REDIST
+
+
+
 int32 RunLiveLinkHub(const TCHAR* CommandLine)
 {
 	FTaskTagScope Scope(ETaskTag::EGameThread);
 
-#if !(UE_BUILD_SHIPPING)
+	// Needs to be initialized early for splash / mount points / plugin search paths.
+	FCommandLine::Set(CommandLine);
+
+	const FText AppName = NSLOCTEXT("LiveLinkHub", "SplashTextName", "LiveLink Hub");
+	FPlatformSplash::SetSplashText(SplashTextType::GameName, *AppName.ToString());
+
+#if !UE_BUILD_SHIPPING
 	if (FParse::Param(CommandLine, TEXT("WaitForDebugger")))
 	{
 		while (!FPlatformMisc::IsDebuggerPresent())
@@ -58,64 +82,48 @@ int32 RunLiveLinkHub(const TCHAR* CommandLine)
 	}
 #endif
 
-	// Must match the path in LiveLinkHubStaging.cs / StageLiveLinkHub.automation.cs
-	const FString RemappedFolder = FPaths::EngineDir() / TEXT("Programs/LiveLinkHub/Staged");
-
-	const FString RemappedConfigs = RemappedFolder / TEXT("Config");
-	const FString RemappedContent = RemappedFolder / TEXT("Content");
-	const FString RemappedPlugins = RemappedFolder / TEXT("Plugins");
-
-	// Needs to be initialized early for mount points / plugin search paths.
-	FCommandLine::Set(CommandLine);
+#if IS_PROGRAM
+	// Disable this when going through PreInit to prevent the console window from appearing.
+	GIsSilent = true;
+#endif
 
-#if COOKED_EDITOR_DISTRIBUTION
+	// Start up the main loop, adding some extra command line arguments:
+#if !IS_PROGRAM
+	int32 Result = GEngineLoop.PreInit(*FString::Printf(TEXT("%s %s"), TEXT(""), CommandLine));
+#else
+	int32 Result = GEngineLoop.PreInit(*FString::Printf(TEXT("%s %s"), CommandLine, TEXT("-RUN=LiveLinkHubCommandlet -Messaging -DDC=NoShared -NoShaderCompile")));
+#endif
 
-	/** If LLH runs in a cooked editor distribution, then we must mount our RemappedEngineFiles folder to the /Engine/ folder in order to
-	 * 	have access to the uncooked content necessary for starting LiveLinkHub.
-	 */
-	if (!IFileManager::Get().DirectoryExists(*RemappedContent))
+	if (Result != 0)
 	{
-		IFileManager::Get().MakeDirectory(*RemappedContent, true);
+		UE_LOG(LogLiveLinkHub, Error, TEXT("EngineLoop PreInit failed (%i)"), Result);
 	}
 
-	if (!IFileManager::Get().DirectoryExists(*RemappedPlugins))
+	if (Result != 0 || IsEngineExitRequested())
 	{
-		IFileManager::Get().MakeDirectory(*RemappedPlugins, true);
+		return Result;
 	}
 
-	/** Only copy the configs if they aren't present to prevent overwriting them. */
-	const FString TargetConfigPath = FPaths::EngineDir() / TEXT("Config/Base.ini");
-	if (!IFileManager::Get().FileExists(*TargetConfigPath) && IFileManager::Get().DirectoryExists(*RemappedConfigs))
+#if !IS_PROGRAM
 	{
-		FPlatformFileManager::Get().GetPlatformFile().CopyDirectoryTree(*(FPaths::EngineDir() / TEXT("Config")), *RemappedConfigs, false);
-	}
+		ILiveLinkHubModule& HubModule = FModuleManager::Get().LoadModuleChecked<ILiveLinkHubModule>("LiveLinkHub");
+		HubModule.PreinitializeLiveLinkHub();
 
-	// Map the engine's content to LiveLinkHub's content folder
-	FPackageName::RegisterMountPoint(TEXT("/Engine/"), RemappedFolder / TEXT("Content"));
-	FPackageName::RegisterMountPoint(TEXT("/Engine/Plugins/"), RemappedFolder / TEXT("Plugins"));
+		Result = GEngineLoop.Init();
 
-	// This is needed to ensure that content generated by the hub will be discoverable (ie. LiveLink recordings)
-	const FString GameContentPath = FPaths::Combine(FPlatformProcess::UserSettingsDir(), *FApp::GetEpicProductIdentifier(), TEXT("LiveLinkHub"), TEXT("Content"));
-	FPackageName::RegisterMountPoint(TEXT("/Game/"), GameContentPath);
-#endif
-
-	// Needed to avoid assertion in scoped slow task
-	GWarn = FPlatformApplicationMisc::GetFeedbackContext();
+		HubModule.StartLiveLinkHub();
 
-	// Add staged plugins to the search path so LLH can find the manifest and content.
-	IPluginManager::Get().AddPluginSearchPath(RemappedPlugins, /*bShouldRefresh*/ true);
+		// Hide the splash screen now that everything is ready to go
+		FPlatformSplash::Hide();
 
-	// Disable this when going through PreInit to prevent the console window from appearing.
-	GIsSilent = true;
-
-	// Start up the main loop, adding some extra command line arguments:
-	const int32 Result = GEngineLoop.PreInit(*FString::Printf(TEXT("%s %s"), CommandLine, TEXT("-RUN=LiveLinkHubCommandlet -Messaging -DDC=NoShared -NoShaderCompile")));
+		while (!IsEngineExitRequested())
+		{
+			GEngineLoop.Tick();
+		}
 
-	if (Result != 0)
-	{
-		UE_LOG(LogLiveLinkHub, Error, TEXT("EngineLoop PreInit failed!"));
-		return Result;
+		HubModule.ShutdownLiveLinkHub();
 	}
+#endif
 
 	GEngineLoop.Exit();
 
diff --git a/Engine/Source/Programs/LiveLinkHub/Private/Platform/Windows/WindowsLiveLinkHubLauncherMain.cpp b/Engine/Source/Programs/LiveLinkHub/Private/Platform/Windows/WindowsLiveLinkHubLauncherMain.cpp
--- a/Engine/Source/Programs/LiveLinkHub/Private/Platform/Windows/WindowsLiveLinkHubLauncherMain.cpp
+++ b/Engine/Source/Programs/LiveLinkHub/Private/Platform/Windows/WindowsLiveLinkHubLauncherMain.cpp
@@ -17,7 +17,7 @@
  * @param nShowCmd Specifies how the window is to be shown.
  * @return Application's exit value.
  */
-int32 WINAPI WinMain( HINSTANCE hInInstance, HINSTANCE hPrevInstance, char* lpCmdLine, int32 nShowCmd )
+int32 WINAPI WinMain(_In_ HINSTANCE hInInstance, _In_opt_ HINSTANCE hPrevInstance, _In_ char* lpCmdLine, _In_ int32 nShowCmd)
 {
 	hInstance = hInInstance;
 
@@ -60,5 +60,8 @@ int32 WINAPI WinMain( HINSTANCE hInInstance, HINSTANCE hPrevInstance, char* lpCm
 		}
 #endif
 	}
+
+	FEngineLoop::AppExit();
+
 	return ErrorLevel;
 }
diff --git a/Engine/Source/Programs/LiveLinkHub/RuntimeDependencies.ini b/Engine/Source/Programs/LiveLinkHub/RuntimeDependencies.ini
--- a/Engine/Source/Programs/LiveLinkHub/RuntimeDependencies.ini
+++ b/Engine/Source/Programs/LiveLinkHub/RuntimeDependencies.ini
@@ -5,9 +5,13 @@
 ; Prefix - denotes exclusion
 +RuntimeDependencies="-/Content/Editor/Slate/Automation/..."
 +RuntimeDependencies="-/Content/Editor/Slate/Testing/..."
++RuntimeDependencies="/Content/EditorMaterials/..."
 +RuntimeDependencies="/Content/EditorMeshes/Camera/..."
++RuntimeDependencies="/Content/EditorResources/..."
 +RuntimeDependencies="/Content/EngineDamageTypes/..."
 +RuntimeDependencies="/Content/EngineFonts/..."
++RuntimeDependencies="/Content/EngineDebugMaterials/..."
++RuntimeDependencies="/Content/EngineMaterials/..."
 +RuntimeDependencies="/Content/Slate/..."
 
 ; Individual assets
@@ -16,74 +20,6 @@
 +RuntimeDependencies="/Content/EditorLandscapeResources/DefaultAlphaTexture.uasset"
 +RuntimeDependencies="/Content/EditorLandscapeResources/SplineEditorMesh.uasset"
 +RuntimeDependencies="/Content/EditorLandscapeResources/SplineEditorMeshMat.uasset"
-+RuntimeDependencies="/Content/EditorMaterials/Camera/CineMat.uasset"
-+RuntimeDependencies="/Content/EditorMaterials/Camera/MI_CineMat_Rig.uasset"
-+RuntimeDependencies="/Content/EditorMaterials/ParticleSystems/PSysThumbnail_NoImage.uasset"
-+RuntimeDependencies="/Content/EditorMaterials/ParticleSystems/PSysThumbnail_OOD.uasset"
-+RuntimeDependencies="/Content/EditorMaterials/WorldPartition/WorldPartitionSpatialHashGridPreviewMaterial.uasset"
-+RuntimeDependencies="/Content/EditorResources/S_Actor.uasset"
-+RuntimeDependencies="/Content/EditorResources/S_BoxReflectionCapture.uasset"
-+RuntimeDependencies="/Content/EditorResources/S_Solver.uasset"
-+RuntimeDependencies="/Content/EditorResources/S_Terrain.uasset"
-+RuntimeDependencies="/Content/EditorResources/S_Trigger.uasset"
-+RuntimeDependencies="/Content/EditorResources/S_TriggerBox.uasset"
-+RuntimeDependencies="/Content/EditorResources/S_TriggerCapsule.uasset"
-+RuntimeDependencies="/Content/EditorResources/S_TriggerSphere.uasset"
-+RuntimeDependencies="/Content/EngineDebugMaterials/DebugMeshMaterial.uasset"
-+RuntimeDependencies="/Content/EngineDebugMaterials/M_SimpleUnlitTranslucent.uasset"
-+RuntimeDependencies="/Content/EngineDebugMaterials/M_VolumeRenderSphereTracePP.uasset"
-+RuntimeDependencies="/Content/EngineDebugMaterials/VolumeToRender.uasset"
-+RuntimeDependencies="/Content/EngineMaterials/BaseFlattenDiffuseMap.uasset"
-+RuntimeDependencies="/Content/EngineMaterials/BaseFlattenDiffuseMap_VT.uasset"
-+RuntimeDependencies="/Content/EngineMaterials/BaseFlattenEmissiveMap.uasset"
-+RuntimeDependencies="/Content/EngineMaterials/BaseFlattenEmissiveMap_VT.uasset"
-+RuntimeDependencies="/Content/EngineMaterials/BaseFlattenGrayscaleMap.uasset"
-+RuntimeDependencies="/Content/EngineMaterials/BaseFlattenGrayscaleMap_VT.uasset"
-+RuntimeDependencies="/Content/EngineMaterials/BaseFlattenLinearColor.uasset"
-+RuntimeDependencies="/Content/EngineMaterials/BaseFlattenLinearColor_VT.uasset"
-+RuntimeDependencies="/Content/EngineMaterials/BaseFlattenNormalMap.uasset"
-+RuntimeDependencies="/Content/EngineMaterials/BaseFlattenNormalMap_VT.uasset"
-+RuntimeDependencies="/Content/EngineMaterials/Black_1x1_EXR_Texture_VT.uasset"
-+RuntimeDependencies="/Content/EngineMaterials/DefaultBokeh.uasset"
-+RuntimeDependencies="/Content/EngineMaterials/DefaultDeferredDecalMaterial.uasset"
-+RuntimeDependencies="/Content/EngineMaterials/DefaultDeferredDecalMaterial.uasset"
-+RuntimeDependencies="/Content/EngineMaterials/DefaultDestructiblePhysicalMaterial.uasset"
-+RuntimeDependencies="/Content/EngineMaterials/DefaultDiffuse.uasset"
-+RuntimeDependencies="/Content/EngineMaterials/DefaultLightFunctionMaterial.uasset"
-+RuntimeDependencies="/Content/EngineMaterials/DefaultLightFunctionMaterial.uasset"
-+RuntimeDependencies="/Content/EngineMaterials/DefaultMaterial.uasset"
-+RuntimeDependencies="/Content/EngineMaterials/DefaultPhysicalMaterial.uasset"
-+RuntimeDependencies="/Content/EngineMaterials/DefaultPostProcessMaterial.uasset"
-+RuntimeDependencies="/Content/EngineMaterials/DefaultPostProcessMaterial.uasset"
-+RuntimeDependencies="/Content/EngineMaterials/DefaultTextMaterialOpaque.uasset"
-+RuntimeDependencies="/Content/EngineMaterials/DefaultWhiteGrid.uasset"
-+RuntimeDependencies="/Content/EngineMaterials/DefaultWhiteGrid_Low.uasset"
-+RuntimeDependencies="/Content/EngineMaterials/EmissiveMeshMaterial.uasset"
-+RuntimeDependencies="/Content/EngineMaterials/FlattenMaterial.uasset"
-+RuntimeDependencies="/Content/EngineMaterials/FlattenMaterial_VT.uasset"
-+RuntimeDependencies="/Content/EngineMaterials/FlattenMaterial_WS_Normal_VT.uasset"
-+RuntimeDependencies="/Content/EngineMaterials/GizmoMaterial.uasset"
-+RuntimeDependencies="/Content/EngineMaterials/Good64x64TilingNoiseHighFreq.uasset"
-+RuntimeDependencies="/Content/EngineMaterials/InvalidLightmapSettings.uasset"
-+RuntimeDependencies="/Content/EngineMaterials/M_InvalidLightmapSettings.uasset"
-+RuntimeDependencies="/Content/EngineMaterials/MiniFont.uasset"
-+RuntimeDependencies="/Content/EngineMaterials/NaniteHiddenSectionMaterial.uasset"
-+RuntimeDependencies="/Content/EngineMaterials/PhAT_JointLimitMaterial.uasset"
-+RuntimeDependencies="/Content/EngineMaterials/PreintegratedSkinBRDF.uasset"
-+RuntimeDependencies="/Content/EngineMaterials/RemoveSurfaceMaterial.uasset"
-+RuntimeDependencies="/Content/EngineMaterials/T_Default_Material_Grid_M.uasset"
-+RuntimeDependencies="/Content/EngineMaterials/T_Default_Material_Grid_N.uasset"
-+RuntimeDependencies="/Content/EngineMaterials/WeightMapArrayPlaceholderTexture.uasset"
-+RuntimeDependencies="/Content/EngineMaterials/WeightMapPlaceholderTexture.uasset"
-+RuntimeDependencies="/Content/EngineMaterials/Widget3DPassThrough.uasset"
-+RuntimeDependencies="/Content/EngineMaterials/Widget3DPassThrough_Masked.uasset"
-+RuntimeDependencies="/Content/EngineMaterials/Widget3DPassThrough_Masked_OneSided.uasset"
-+RuntimeDependencies="/Content/EngineMaterials/Widget3DPassThrough_Opaque.uasset"
-+RuntimeDependencies="/Content/EngineMaterials/Widget3DPassThrough_Opaque_OneSided.uasset"
-+RuntimeDependencies="/Content/EngineMaterials/Widget3DPassThrough_Translucent.uasset"
-+RuntimeDependencies="/Content/EngineMaterials/Widget3DPassThrough_Translucent_OneSided.uasset"
-+RuntimeDependencies="/Content/EngineMaterials/WorldGridMaterial.uasset"
-+RuntimeDependencies="/Content/EngineMaterials/WorldGridMaterial.uasset"
 +RuntimeDependencies="/Content/EngineMeshes/Sphere.uasset"
 +RuntimeDependencies="/Content/EngineResources/Black.uasset"
 +RuntimeDependencies="/Content/EngineResources/DefaultTexture.uasset"
diff --git a/Engine/Source/Programs/LiveLinkHubEditor/Config/DefaultEditorPerProjectUserSettings.ini b/Engine/Source/Programs/LiveLinkHubEditor/Config/DefaultEditorPerProjectUserSettings.ini
new file mode 100644
--- /dev/null
+++ b/Engine/Source/Programs/LiveLinkHubEditor/Config/DefaultEditorPerProjectUserSettings.ini
@@ -0,0 +1,2 @@
+[/Script/UnrealEd.EditorLoadingSavingSettings]
+bAutoSaveEnable=False
diff --git a/Engine/Source/Programs/LiveLinkHubEditor/Config/DefaultGame.ini b/Engine/Source/Programs/LiveLinkHubEditor/Config/DefaultGame.ini
--- a/Engine/Source/Programs/LiveLinkHubEditor/Config/DefaultGame.ini
+++ b/Engine/Source/Programs/LiveLinkHubEditor/Config/DefaultGame.ini
@@ -1,6 +1,6 @@
 [CookedEditorSettings_CookedEditor]
 CookedEditorTargetName=LiveLinkHubEditor
-bStageShaderDirs=false
+bStageShaderDirs=true
 bStagePlatformBuildDirs=false
 bStageExtrasDirs=false
 bStagePlatformDirs=false
