LiveLinkHub: Build fix for non-Windows platforms.

diff --git a/Engine/Source/Programs/LiveLinkHub/Private/LiveLinkHubRun.cpp b/Engine/Source/Programs/LiveLinkHub/Private/LiveLinkHubRun.cpp
--- a/Engine/Source/Programs/LiveLinkHub/Private/LiveLinkHubRun.cpp
+++ b/Engine/Source/Programs/LiveLinkHub/Private/LiveLinkHubRun.cpp
@@ -28,18 +28,6 @@
 DEFINE_LOG_CATEGORY_STATIC(LogLiveLinkHub, Log, All);
 
 
-// Opt in to new D3D12 redist and tell the loader where to search for D3D12Core.dll.
-// The D3D loader looks for these symbol exports in the .exe module.
-// We only support this on x64 Windows Desktop platforms. Other platforms or non-redist-aware 
-// versions of Windows will transparently load default OS-provided D3D12 library.
-#define USE_D3D12_REDIST (PLATFORM_DESKTOP && PLATFORM_CPU_X86_FAMILY && PLATFORM_64BITS && 1)
-#if USE_D3D12_REDIST
-extern "C" { _declspec(dllexport) extern const UINT D3D12SDKVersion = 614; } // D3D12_SDK_VERSION
-extern "C" { _declspec(dllexport) extern const char* D3D12SDKPath = ".\\D3D12\\"; }
-#endif // USE_D3D12_REDIST
-
-
-
 int32 RunLiveLinkHub(const TCHAR* CommandLine)
 {
 	FTaskTagScope Scope(ETaskTag::EGameThread);
diff --git a/Engine/Source/Programs/LiveLinkHub/Private/Platform/Windows/WindowsLiveLinkHubLauncherMain.cpp b/Engine/Source/Programs/LiveLinkHub/Private/Platform/Windows/WindowsLiveLinkHubLauncherMain.cpp
--- a/Engine/Source/Programs/LiveLinkHub/Private/Platform/Windows/WindowsLiveLinkHubLauncherMain.cpp
+++ b/Engine/Source/Programs/LiveLinkHub/Private/Platform/Windows/WindowsLiveLinkHubLauncherMain.cpp
@@ -8,6 +8,18 @@
 #include "Misc/CommandLine.h"
 #include "Misc/OutputDeviceError.h"
 
+
+// Opt in to new D3D12 redist and tell the loader where to search for D3D12Core.dll.
+// The D3D loader looks for these symbol exports in the .exe module.
+// We only support this on x64 Windows Desktop platforms. Other platforms or non-redist-aware 
+// versions of Windows will transparently load default OS-provided D3D12 library.
+#define USE_D3D12_REDIST (PLATFORM_DESKTOP && PLATFORM_CPU_X86_FAMILY && PLATFORM_64BITS && 1)
+#if USE_D3D12_REDIST
+extern "C" { _declspec(dllexport) extern const UINT D3D12SDKVersion = 614; } // D3D12_SDK_VERSION
+extern "C" { _declspec(dllexport) extern const char* D3D12SDKPath = ".\\D3D12\\"; }
+#endif // USE_D3D12_REDIST
+
+
 /**
  * The main application entry point for Windows platforms.
  *
