LiveLinkHubEditor - Disable CPU throttling using EditorEngine's ShouldDisableCPUThrottlingDelegates

diff --git a/Engine/Source/Programs/LiveLinkHub/Private/LiveLinkHubRun.cpp b/Engine/Source/Programs/LiveLinkHub/Private/LiveLinkHubRun.cpp
--- a/Engine/Source/Programs/LiveLinkHub/Private/LiveLinkHubRun.cpp
+++ b/Engine/Source/Programs/LiveLinkHub/Private/LiveLinkHubRun.cpp
@@ -2,6 +2,8 @@
 
 #include "LiveLinkHubRun.h"
 
+#include "Editor.h"
+#include "Editor/EditorEngine.h"
 #include "Misc/App.h"
 #include "Misc/PackageName.h"
 #include "Misc/Paths.h"
@@ -109,6 +111,8 @@ int32 RunLiveLinkHub(const TCHAR* CommandLine)
 
 		Result = GEngineLoop.Init();
 
+		GEditor->ShouldDisableCPUThrottlingDelegates.Add(UEditorEngine::FShouldDisableCPUThrottling::CreateLambda([]() { return true; }));
+
 		HubModule.StartLiveLinkHub(bLauncherDistribution);
 
 		// Hide the splash screen now that everything is ready to go
