LiveLink Hub: Add Plugin Browser and LiveLinkXR support. Also adds automation testing, slate reflector, and live coding.

diff --git a/Engine/Plugins/Animation/LiveLinkHub/Source/LiveLinkHub/LiveLinkHub.Build.cs b/Engine/Plugins/Animation/LiveLinkHub/Source/LiveLinkHub/LiveLinkHub.Build.cs
--- a/Engine/Plugins/Animation/LiveLinkHub/Source/LiveLinkHub/LiveLinkHub.Build.cs
+++ b/Engine/Plugins/Animation/LiveLinkHub/Source/LiveLinkHub/LiveLinkHub.Build.cs
@@ -7,13 +7,13 @@ public class LiveLinkHub : ModuleRules
 	public LiveLinkHub(ReadOnlyTargetRules Target) : base(Target)
 	{
 		PublicDependencyModuleNames.AddRange(
-		new string[]
-		{
-			"Core",
-			"CoreUObject",
-		});
+			new string[]
+			{
+				"Core",
+				"CoreUObject",
+			});
 
-			PrivateDependencyModuleNames.AddRange(
+		PrivateDependencyModuleNames.AddRange(
 			new string[]
 			{
 				"ApplicationCore",
@@ -46,5 +46,27 @@ public class LiveLinkHub : ModuleRules
 				"UnrealEd",
 				"OutputLog",
 			});
+
+		if (Target.Configuration != UnrealTargetConfiguration.Shipping)
+		{
+			PrivateIncludePathModuleNames.AddRange(
+				new string[] {
+					"SlateReflector",
+				});
+
+			DynamicallyLoadedModuleNames.AddRange(
+				new string[] {
+					"SlateReflector",
+				});
+
+			if (Target.bBuildTargetDeveloperTools)
+			{
+				PrivateDependencyModuleNames.AddRange(
+					new string[] {
+						"SessionServices",
+						"SessionFrontend",
+					});
+			}
+		}
 	}
 }
diff --git a/Engine/Plugins/Animation/LiveLinkHub/Source/LiveLinkHub/Private/UI/Widgets/SLiveLinkHubSettings.h b/Engine/Plugins/Animation/LiveLinkHub/Source/LiveLinkHub/Private/UI/Widgets/SLiveLinkHubSettings.h
--- a/Engine/Plugins/Animation/LiveLinkHub/Source/LiveLinkHub/Private/UI/Widgets/SLiveLinkHubSettings.h
+++ b/Engine/Plugins/Animation/LiveLinkHub/Source/LiveLinkHub/Private/UI/Widgets/SLiveLinkHubSettings.h
@@ -137,13 +137,20 @@ public:
 
 		MenuBuilder.BeginSection("HubSettings", LOCTEXT("HubSettingsLabel", "Hub Settings"));
 		MenuBuilder.AddMenuEntry(
-			LOCTEXT("SettingsMenuEntry", "Settings"),
-			LOCTEXT("SettingsMenuEntry_ToolTip", "Open the settings menu."),
-			FSlateIcon(),
+			LOCTEXT("SettingsMenuEntry", "Settings..."),
+			LOCTEXT("SettingsMenuEntry_ToolTip", "Open the Settings tab."),
+			FSlateIcon(FAppStyle::GetAppStyleSetName(), "EditorPreferences.TabIcon"),
 			OpenHubSettingsAction,
 			NAME_None,
 			EUserInterfaceActionType::Button);
-
+		MenuBuilder.AddMenuEntry(
+			LOCTEXT("PluginsMenuEntry", "Plugins..."),
+			LOCTEXT("PluginsMenuEntry_ToolTip", "Open the Plugins Browser tab."),
+			FSlateIcon("PluginStyle", "Plugins.TabIcon"),
+			FExecuteAction::CreateStatic([]()
+			{
+				FGlobalTabmanager::Get()->TryInvokeTab(FName("PluginsEditor"));
+			}));
 		MenuBuilder.EndSection();
 
 		return MenuBuilder.MakeWidget();
diff --git a/Engine/Plugins/Animation/LiveLinkHub/Source/LiveLinkHub/Private/UI/Widgets/SLiveLinkHubTabViewWithManagerBase.cpp b/Engine/Plugins/Animation/LiveLinkHub/Source/LiveLinkHub/Private/UI/Widgets/SLiveLinkHubTabViewWithManagerBase.cpp
--- a/Engine/Plugins/Animation/LiveLinkHub/Source/LiveLinkHub/Private/UI/Widgets/SLiveLinkHubTabViewWithManagerBase.cpp
+++ b/Engine/Plugins/Animation/LiveLinkHub/Source/LiveLinkHub/Private/UI/Widgets/SLiveLinkHubTabViewWithManagerBase.cpp
@@ -6,9 +6,13 @@
 #include "Framework/Commands/UICommandList.h"
 #include "Framework/Docking/LayoutService.h"
 #include "Framework/MultiBox/MultiBoxBuilder.h"
+#include "ISessionFrontendModule.h"
+#include "ISlateReflectorModule.h"
 #include "LiveLinkHub.h"
 #include "LiveLinkHubCommands.h"
 #include "LiveLinkHubModule.h"
+#include "Misc/CommandLine.h"
+#include "Misc/EngineBuildSettings.h"
 #include "Modules/ModuleManager.h"
 #include "UI/Widgets/SLiveLinkHubTabViewBase.h"
 
@@ -90,6 +94,39 @@ void SLiveLinkHubTabViewWithManagerBase::FillFileMenu(FMenuBuilder& MenuBuilder)
 	MenuBuilder.AddMenuEntry(FLiveLinkHubCommands::Get().SaveConfig);
 	MenuBuilder.AddMenuEntry(FLiveLinkHubCommands::Get().SaveConfigAs);
 	MenuBuilder.EndSection();
+
+#if !UE_BUILD_SHIPPING
+	const bool bShowDevTools = FParse::Param(FCommandLine::Get(), TEXT("Development"))
+		|| FEngineBuildSettings::IsInternalBuild();
+
+	if (bShowDevTools)
+	{
+		MenuBuilder.BeginSection("Development", LOCTEXT("DevelopmentHeader", "Development"));
+		MenuBuilder.AddMenuEntry(
+			LOCTEXT("FileMenu_Development_AutomationTools", "Automation Tools"),
+			FText(),
+			FSlateIcon(),
+			FExecuteAction::CreateStatic([]()
+			{
+				ISessionFrontendModule& SessionFrontend =
+					FModuleManager::LoadModuleChecked<ISessionFrontendModule>("SessionFrontend");
+				SessionFrontend.InvokeSessionFrontend(FName("AutomationPanel"));
+			})
+		);
+		MenuBuilder.AddMenuEntry(
+			LOCTEXT("FileMenu_Development_WidgetReflector", "Widget Reflector"),
+			FText(),
+			FSlateIcon(),
+			FExecuteAction::CreateStatic([]()
+			{
+				ISlateReflectorModule& SlateReflector =
+					FModuleManager::LoadModuleChecked<ISlateReflectorModule>("SlateReflector");
+				SlateReflector.DisplayWidgetReflector();
+			})
+		);
+		MenuBuilder.EndSection();
+	}
+#endif // #if !UE_BUILD_SHIPPING
 }
 
 void SLiveLinkHubTabViewWithManagerBase::FillWindowMenu(FMenuBuilder& MenuBuilder)
diff --git a/Engine/Source/Programs/LiveLinkHub/LiveLinkHubLauncher.Build.cs b/Engine/Source/Programs/LiveLinkHub/LiveLinkHubLauncher.Build.cs
--- a/Engine/Source/Programs/LiveLinkHub/LiveLinkHubLauncher.Build.cs
+++ b/Engine/Source/Programs/LiveLinkHub/LiveLinkHubLauncher.Build.cs
@@ -14,7 +14,11 @@ public class LiveLinkHubLauncher : ModuleRules
 {
 	public LiveLinkHubLauncher(ReadOnlyTargetRules Target) : base(Target)
 	{
-		PrivateDependencyModuleNames.Add("LiveLinkHub");
+		PrivateDependencyModuleNames.AddRange(
+			new string[] {
+				"LiveLinkHub",
+				"PluginBrowser",
+			});
 
 		// LaunchEngineLoop dependencies
 		PrivateDependencyModuleNames.AddRange(
@@ -63,13 +67,23 @@ public class LiveLinkHubLauncher : ModuleRules
 		// LaunchEngineLoop editor dependencies
 		PrivateDependencyModuleNames.AddRange(
 			new string[] {
-				
 				"PropertyEditor",
 				"DerivedDataCache",
 				"ToolWidgets",
 				"UnrealEd"
 		});
 
+
+		if (Target.Configuration != UnrealTargetConfiguration.Shipping)
+		{
+			PrivateDependencyModuleNames.AddRange(
+				new string[] {
+					"AutomationController",
+					"AutomationTest",
+					"AutomationWorker",
+				});
+		}
+
 		if (Target.IsInPlatformGroup(UnrealPlatformGroup.Linux))
 		{
 			PrivateDependencyModuleNames.AddRange(
diff --git a/Engine/Source/Programs/LiveLinkHub/Private/LiveLinkHubPluginHelpers.cpp b/Engine/Source/Programs/LiveLinkHub/Private/LiveLinkHubPluginHelpers.cpp
new file mode 100644
--- /dev/null
+++ b/Engine/Source/Programs/LiveLinkHub/Private/LiveLinkHubPluginHelpers.cpp
@@ -0,0 +1,93 @@
+// Copyright Epic Games, Inc. All Rights Reserved.
+
+#include "LiveLinkHubPluginHelpers.h"
+#include "HAL/PlatformProcess.h"
+#include "Interfaces/IPluginManager.h"
+#include "IPluginBrowser.h"
+#include "Misc/App.h"
+#include "Misc/ConfigCacheIni.h"
+
+
+namespace LiveLinkHub::PluginHelpers
+{
+	const TCHAR* const ConfigSection = TEXT("LiveLinkHub");
+	const TCHAR* const ConfigKey = TEXT("PluginDirectories");
+}
+
+
+FString LiveLinkHub::PluginHelpers::GetPluginConfigPath()
+{
+	return FPaths::Combine(
+		FPlatformProcess::UserSettingsDir(),
+		*FApp::GetEpicProductIdentifier(),
+		TEXT("LiveLinkHub"),
+		TEXT("Config"),
+		TEXT("LiveLinkHubPlugins.ini")
+	);
+}
+
+void LiveLinkHub::PluginHelpers::ReadPluginDirectoriesFromConfig(TArray<FString>& OutDirectories)
+{
+	FConfigFile ConfigFile;
+	ConfigFile.Read(GetPluginConfigPath());
+	ConfigFile.GetArray(ConfigSection, ConfigKey, OutDirectories);
+}
+
+void LiveLinkHub::PluginHelpers::WritePluginDirectoriesToConfig(const TArray<FString>& InDirectories)
+{
+	FConfigFile ConfigFile;
+	ConfigFile.Read(GetPluginConfigPath());
+
+	ConfigFile.ResetKeyInSection(ConfigSection, ConfigKey);
+	for (const FString& Directory : InDirectories)
+	{
+		ConfigFile.AddUniqueToSection(ConfigSection, ConfigKey, Directory);
+	}
+
+	ConfigFile.Write(GetPluginConfigPath());
+}
+
+void LiveLinkHub::PluginHelpers::RestoreSavedPluginDirectories()
+{
+	TArray<FString> SavedDirectories;
+	ReadPluginDirectoriesFromConfig(SavedDirectories);
+
+	bool bPluginPathsChanged = false;
+
+	IPluginManager& PluginManager = IPluginManager::Get();
+	for (const FString& Directory : SavedDirectories)
+	{
+		const bool bRefresh_false = false;
+		bPluginPathsChanged |= PluginManager.AddPluginSearchPath(Directory, bRefresh_false);
+	}
+
+	if (bPluginPathsChanged)
+	{
+		PluginManager.RefreshPluginsList();
+	}
+}
+
+void LiveLinkHub::PluginHelpers::HandlePluginDirectoriesChanged()
+{
+	IPluginManager& PluginManager = IPluginManager::Get();
+
+	TSet<FExternalPluginPath> PluginSources;
+	PluginManager.GetExternalPluginSources(PluginSources);
+
+	TArray<FString> UserDirectories;
+	for (const FExternalPluginPath& Path : PluginSources)
+	{
+		if (Path.Source == EPluginExternalSource::Other)
+		{
+			UserDirectories.Add(Path.Path);
+		}
+	}
+
+	WritePluginDirectoriesToConfig(UserDirectories);
+}
+
+void LiveLinkHub::PluginHelpers::RegisterPluginDirectoriesChangedHandler()
+{
+	IPluginBrowser& PluginBrowser = IPluginBrowser::Get();
+	PluginBrowser.OnPluginDirectoriesChanged().AddStatic(HandlePluginDirectoriesChanged);
+}
diff --git a/Engine/Source/Programs/LiveLinkHub/Private/LiveLinkHubPluginHelpers.h b/Engine/Source/Programs/LiveLinkHub/Private/LiveLinkHubPluginHelpers.h
new file mode 100644
--- /dev/null
+++ b/Engine/Source/Programs/LiveLinkHub/Private/LiveLinkHubPluginHelpers.h
@@ -0,0 +1,26 @@
+// Copyright Epic Games, Inc. All Rights Reserved.
+
+#include "Containers/Array.h"
+#include "Containers/UnrealString.h"
+
+
+namespace LiveLinkHub::PluginHelpers
+{
+	/** Returns the path to LiveLinkHubPlugins.ini. */
+	FString GetPluginConfigPath();
+
+	/** Returns the user plugin directories stored in LiveLinkHubPlugins.ini. */
+	void ReadPluginDirectoriesFromConfig(TArray<FString>& OutDirectories);
+
+	/** Serializes the specified user plugin directories out to LiveLinkHubPlugins.ini. */
+	void WritePluginDirectoriesToConfig(const TArray<FString>& InDirectories);
+
+	/** Reads the configured user plugin directories and adds them to the plugin manager search path. */
+	void RestoreSavedPluginDirectories();
+
+	/** Stores the latest plugin manager search paths to disk. */
+	void HandlePluginDirectoriesChanged();
+
+	/** Connects the `HandlePluginDirectoriesChanged` handler to the plugin browser directories changed event. */
+	void RegisterPluginDirectoriesChangedHandler();
+}
diff --git a/Engine/Source/Programs/LiveLinkHub/Private/LiveLinkHubRun.cpp b/Engine/Source/Programs/LiveLinkHub/Private/LiveLinkHubRun.cpp
--- a/Engine/Source/Programs/LiveLinkHub/Private/LiveLinkHubRun.cpp
+++ b/Engine/Source/Programs/LiveLinkHub/Private/LiveLinkHubRun.cpp
@@ -10,11 +10,13 @@
 #include "HAL/FileManager.h"
 #include "HAL/PlatformApplicationMisc.h"
 #include "HAL/PlatformFileManager.h"
+#include "HAL/PlatformMisc.h"
 #include "HAL/PlatformSplash.h"
 #include "ILiveLinkHubModule.h"
-#include "Interfaces/IPluginManager.h"
 #include "LaunchEngineLoop.h"
+#include "LiveLinkHubPluginHelpers.h"
 #include "Misc/CommandLine.h"
+#include "Misc/CoreDelegates.h"
 #include "Misc/Paths.h"
 #include "Modules/ModuleManager.h"
 
@@ -22,10 +24,6 @@
 #define WITH_ASSET_LOADING_AUDIT 0
 #endif
 
-#if WITH_ASSET_LOADING_AUDIT
-#	include "Misc/CoreDelegates.h"
-#endif
-
 
 DEFINE_LOG_CATEGORY_STATIC(LogLiveLinkHub, Log, All);
 
@@ -73,7 +71,6 @@ int32 RunLiveLinkHub(const TCHAR* CommandLine)
 		bLauncherDistribution = true;
 		FPaths::SetProjectFilePath(StagedProjectPath);
 	}
-#endif
 
 	if (bLauncherDistribution)
 	{
@@ -82,6 +79,22 @@ int32 RunLiveLinkHub(const TCHAR* CommandLine)
 		FPackageName::RegisterMountPoint(TEXT("/Game/"), GameContentPath);
 	}
 
+	// Used by Live Coding, among other things.
+	FPlatformMisc::SetUBTTargetName(TEXT("LiveLinkHubEditor"));
+
+	// We need to specify this manually, because it isn't inferred from the command line.
+	// But we can't set it too early, because LaunchSetGameName clears it.
+	FCoreDelegates::OnInit.AddStatic([]()
+	{
+		// Required for OpenXR instance registration.
+		FApp::SetProjectName(TEXT("LiveLinkHub"));
+	});
+
+	// Plugin directory config save/load
+	FCoreDelegates::TSConfigReadyForUse().AddStatic(LiveLinkHub::PluginHelpers::RestoreSavedPluginDirectories);
+	FCoreDelegates::OnAllModuleLoadingPhasesComplete.AddStatic(LiveLinkHub::PluginHelpers::RegisterPluginDirectoriesChangedHandler);
+#endif
+
 #if IS_PROGRAM
 	// Disable this when going through PreInit to prevent the console window from appearing.
 	GIsSilent = true;
@@ -89,7 +102,7 @@ int32 RunLiveLinkHub(const TCHAR* CommandLine)
 
 	// Start up the main loop, adding some extra command line arguments:
 #if !IS_PROGRAM
-	int32 Result = GEngineLoop.PreInit(*FString::Printf(TEXT("%s %s"), TEXT(""), CommandLine));
+	int32 Result = GEngineLoop.PreInit(*FString::Printf(TEXT("%s %s"), CommandLine, TEXT("-xrtrackingonly")));
 #else
 	int32 Result = GEngineLoop.PreInit(*FString::Printf(TEXT("%s %s"), CommandLine, TEXT("-RUN=LiveLinkHubCommandlet -Messaging -DDC=NoShared -NoShaderCompile")));
 #endif
diff --git a/Engine/Source/Programs/LiveLinkHubEditor/Config/DefaultEditor.ini b/Engine/Source/Programs/LiveLinkHubEditor/Config/DefaultEditor.ini
new file mode 100644
--- /dev/null
+++ b/Engine/Source/Programs/LiveLinkHubEditor/Config/DefaultEditor.ini
@@ -0,0 +1,14 @@
+[EditorSettings]
+; We don't allow creating new plugins.
+bCanCreatePluginsFromBrowser=false
+; We don't allow toggling individual plugins.
+bCanModifyPluginsFromBrowser=false
+
+; We allow modifying user plugin paths...
+bCanModifyUserPluginDirectoriesFromBrowser=true
+; ...but not paths stored in the .uproject descriptor.
+bCanModifyProjectPluginDirectoriesFromBrowser=false
+
+; Hide plugins that aren't "Installed" or "External" to cut down on noise.
++HidePluginCategoriesFromBrowser=All
++HidePluginCategoriesFromBrowser=Built-In
diff --git a/Engine/Source/Programs/LiveLinkHubEditor/Source/LiveLinkHubEditor.Target.cs b/Engine/Source/Programs/LiveLinkHubEditor/Source/LiveLinkHubEditor.Target.cs
--- a/Engine/Source/Programs/LiveLinkHubEditor/Source/LiveLinkHubEditor.Target.cs
+++ b/Engine/Source/Programs/LiveLinkHubEditor/Source/LiveLinkHubEditor.Target.cs
@@ -26,6 +26,8 @@ public class LiveLinkHubEditorTarget : TargetRules
 
 		SolutionDirectory = "Programs/LiveLink";
 		LaunchModuleName = "LiveLinkHubLauncher";
+
+		DefaultBuildSettings = BuildSettingsVersion.Latest;
 		IncludeOrderVersion = EngineIncludeOrderVersion.Latest;
 
 		// These plugins are required for running LiveLinkHub. 
@@ -35,6 +37,8 @@ public class LiveLinkHubEditorTarget : TargetRules
 			"ACLPlugin",
 			"LiveLink",
 			"LiveLinkHub",
+			"PluginBrowser",
+
 			"LiveLinkCamera",
 			"LiveLinkLens",
 			"LensComponent",
@@ -47,7 +51,10 @@ public class LiveLinkHubEditorTarget : TargetRules
 			"UdpMessaging",
 			"CameraCalibrationCore",
 			"AppleARKitFaceSupport",
-			"XInputDevice"
+			"XInputDevice",
+
+			"LiveLinkXR",
+			"OpenXRViveTracker",
 		});
 
 		if (bEnableCaptureManagerPlugin)
