LiveLinkHub: Remove unused program target, rename LiveLinkHubEditor to LiveLinkHub.

#rnx

diff --git a/Engine/Build/InstalledEngineBuild.xml b/Engine/Build/InstalledEngineBuild.xml
--- a/Engine/Build/InstalledEngineBuild.xml
+++ b/Engine/Build/InstalledEngineBuild.xml
@@ -196,8 +196,8 @@
 			<Compile Target="UnrealEditor" Platform="Win64" Configuration="DebugGame" Tag="#UnrealEditor Win64" Arguments="-precompile -allmodules $(ExtraCompileArgs) $(Win64NoPCHArg)"/>
 			<Compile Target="UnrealEditor" Platform="Win64" Configuration="Development" Tag="#UnrealEditor Win64" Arguments="-precompile -allmodules $(ExtraCompileArgs) $(Win64NoPCHArg)"/>
 			
-			<!-- Note: We compile LiveLinkHubEditor here because we share binaries with the editor. This drastically reduces its build time and prevents having duplicate files in the MakeInstalledBuild step. -->
-			<Compile Target="LiveLinkHubEditor" Configuration="Development" Platform="Win64" Tag="#UnrealEditor Win64" Arguments="-precompile $(ExtraCompileArgs) $(Win64NoPCHArg)" />
+			<!-- Note: We compile LiveLinkHub here because we share a modular build environment with the editor. This drastically reduces its build time and prevents having duplicate files in the MakeInstalledBuild step. -->
+			<Compile Target="LiveLinkHub" Configuration="Development" Platform="Win64" Tag="#UnrealEditor Win64" Arguments="-precompile $(ExtraCompileArgs) $(Win64NoPCHArg)" />
 
 			<Do If="$(EmbedSrcSrvInfo)">
 				<!-- Embed source info into the PDB files. Should be done from this machine to ensure that paths are correct. -->
@@ -239,9 +239,9 @@
 
 			<Compile Target="UnrealEditor" Platform="Mac" Configuration="DebugGame" Tag="#UnrealEditor Mac" Arguments="-precompile -EnableDSYM $(ExtraCompileArgsMac) $(AllModulesArg) $(MacArchitectureArg)"/>
 			<Compile Target="UnrealEditor" Platform="Mac" Configuration="Development" Tag="#UnrealEditor Mac" Arguments="-precompile -EnableDSYM $(ExtraCompileArgsMac) $(AllModulesArg) $(MacArchitectureArg)"/>
-			
-			<!-- Note: We compile LiveLinkHubEditor here because we share binaries with the editor. This drastically reduces its build time and prevents having duplicate files in it -->
-			<!--<Compile Target="LiveLinkHubEditor" Configuration="Development" Platform="Mac" Tag="#UnrealEditor Mac" Arguments="-precompile -EnableDSYM $(MacArchitectureArg)" /> -->
+
+			<!-- Note: We compile LiveLinkHub here because we share a modular build environment with the editor. This drastically reduces its build time and prevents having duplicate files in the MakeInstalledBuild step. -->
+			<!--<Compile Target="LiveLinkHub" Configuration="Development" Platform="Mac" Tag="#UnrealEditor Mac" Arguments="-precompile -EnableDSYM $(MacArchitectureArg)" /> -->
 
 			<Property Name="UnrealEditorExceptions" Value=""/>
 			<Property Name="UnrealEditorExceptions" Value="$(PluginsExceptions)" If="$(CompileDatasmithPlugins)"/>
@@ -274,9 +274,9 @@
 		<Node Name="Compile UnrealEditor Linux" Requires="Update Version Files" Produces="#UnrealEditor Linux;#UnrealEditor Linux Unstripped;#UnrealEditor Linux Stripped;#UnrealEditor Linux WithExceptions">
 			<Compile Target="UnrealEditor" Platform="Linux" Configuration="DebugGame" Tag="#UnrealEditor Linux" Arguments="-precompile -allmodules $(ExtraCompileArgs)"/>
 			<Compile Target="UnrealEditor" Platform="Linux" Configuration="Development" Tag="#UnrealEditor Linux" Arguments="-precompile -allmodules $(ExtraCompileArgs)"/>
-			
-			<!-- Note: We compile LiveLinkHubEditor here because we share binaries with the editor. This drastically reduces its build time and prevents having duplicate files in the MakeInstalledBuild step. -->
-			<!--<Compile Target="LiveLinkHubEditor" Configuration="Development" Platform="Linux" Tag="#UnrealEditor Linux" Arguments="-precompile $(ExtraCompileArgs)" />-->
+
+			<!-- Note: We compile LiveLinkHub here because we share a modular build environment with the editor. This drastically reduces its build time and prevents having duplicate files in the MakeInstalledBuild step. -->
+			<!--<Compile Target="LiveLinkHub" Configuration="Development" Platform="Linux" Tag="#UnrealEditor Linux" Arguments="-precompile $(ExtraCompileArgs)" />-->
 
 			<Tag Files="#UnrealEditor Linux" Filter="$(LinuxStripFilter)" With="#UnrealEditor Linux Unstripped"/>
 			<Tag Files="#UnrealEditor Linux" Except="$(PluginsExceptions)" With="#UnrealEditor Linux WithExceptions"/>
diff --git a/Engine/Programs/LiveLinkHub/Config/DefaultEditor.ini b/Engine/Programs/LiveLinkHub/Config/DefaultEditor.ini
deleted file mode 100644
diff --git a/Engine/Programs/LiveLinkHub/Config/DefaultEditorPerProjectUserSettings.ini b/Engine/Programs/LiveLinkHub/Config/DefaultEditorPerProjectUserSettings.ini
deleted file mode 100644
diff --git a/Engine/Programs/LiveLinkHub/Config/DefaultEngine.ini b/Engine/Programs/LiveLinkHub/Config/DefaultEngine.ini
deleted file mode 100644
diff --git a/Engine/Programs/LiveLinkHub/Config/DefaultGame.ini b/Engine/Programs/LiveLinkHub/Config/DefaultGame.ini
deleted file mode 100644
diff --git a/Engine/Programs/LiveLinkHub/LiveLinkHub.uproject b/Engine/Programs/LiveLinkHub/LiveLinkHub.uproject
deleted file mode 100644
diff --git a/Engine/Programs/UnrealGameSync/UnrealGameSync.ini b/Engine/Programs/UnrealGameSync/UnrealGameSync.ini
--- a/Engine/Programs/UnrealGameSync/UnrealGameSync.ini
+++ b/Engine/Programs/UnrealGameSync/UnrealGameSync.ini
@@ -37,9 +37,8 @@ BuildUnrealHeaderTool=0
 +Step=(UniqueId="71e65719-b7bb-49c5-961b-620acf1e1dad", Description="Compile SwitchboardListenerHelper", StatusText="Compiling SwitchboardListenerHelper...", EstimatedDuration="1", Type="Compile", Target="SwitchboardListenerHelper", Platform="Win64", Configuration="Development", Arguments="", bNormalSync="False", bScheduledSync="False", bShowAsTool="False")
 +Step=(UniqueId="5036c75b-8df0-4329-82a1-869dd2d48605", Description="Compile UnrealMultiUserServer", StatusText="Compiling UnrealMultiUserServer...", EstimatedDuration="1", Type="Compile", Target="UnrealMultiUserServer", Platform="Win64", Configuration="Development", Arguments="", bNormalSync="False", bScheduledSync="False", bShowAsTool="False")
 +Step=(UniqueId="274b89c3-9dc6-4465-a508-40abc4593cc2", Description="Compile UnrealMultiUserSlateServer", StatusText="Compiling UnrealMultiUserSlateServer...", EstimatedDuration="1", Type="Compile", Target="UnrealMultiUserSlateServer", Platform="Win64", Configuration="Development", Arguments="", bNormalSync="False", bScheduledSync="False", bShowAsTool="False")
-+Step=(UniqueId="affbea0c-0c41-42cb-abd4-0643ff256cae", Description="Compile LiveLinkHub", StatusText="Compiling LiveLinkHub...", EstimatedDuration="1", Type="Compile", Target="LiveLinkHub", Platform="Win64", Configuration="Development", Arguments="", bNormalSync="False", bScheduledSync="False", bShowAsTool="False")
-+Step=(UniqueId="fffcbd4b-0f14-40ca-9069-f09c1c6e5898", Description="Compile LiveLinkHubEditor", StatusText="Compiling LiveLinkHubEditor...", EstimatedDuration="1", Type="Compile", Target="LiveLinkHubEditor", Platform="Win64", Configuration="Development", Arguments="", bNormalSync="False", bScheduledSync="False", bShowAsTool="False")
-+Step=(UniqueId="0aadd9c7-17b7-4818-a9ac-d11dc71f6ab9", Description="Launch LiveLinkHubEditor", StatusText="Running LiveLinkHubEditor", EstimatedDuration="1", Type="Other", FileName="$(BranchDir)\\Engine\\Binaries\\Win64\\LiveLinkHubEditor.exe", WorkingDir="", Arguments="", bUseLogWindow="False", bNormalSync="False", bShowAsTool="True", bScheduledSync="False", Requires="fffcbd4b-0f14-40ca-9069-f09c1c6e5898" )
++Step=(UniqueId="fffcbd4b-0f14-40ca-9069-f09c1c6e5898", Description="Compile LiveLinkHub", StatusText="Compiling LiveLinkHub...", EstimatedDuration="1", Type="Compile", Target="LiveLinkHub", Platform="Win64", Configuration="Development", Arguments="", bNormalSync="False", bScheduledSync="False", bShowAsTool="False")
++Step=(UniqueId="0aadd9c7-17b7-4818-a9ac-d11dc71f6ab9", Description="Launch LiveLinkHub", StatusText="Running LiveLinkHub", EstimatedDuration="1", Type="Other", FileName="$(BranchDir)\\Engine\\Binaries\\Win64\\LiveLinkHub.exe", WorkingDir="", Arguments="", bUseLogWindow="False", bNormalSync="False", bShowAsTool="True", bScheduledSync="False", Requires="fffcbd4b-0f14-40ca-9069-f09c1c6e5898" )
 +Step=(UniqueId="f7f0c4c7-5551-4485-a961-f3e4fd9d138c", Description="Compile InterchangeWorker", StatusText="Compiling InterchangeWorker...", EstimatedDuration="1", Type="Compile", Target="InterchangeWorker", Platform="Win64", Configuration="Development", Arguments="", bNormalSync="False", bScheduledSync="False", bShowAsTool="False")
 
 [Default]
diff --git a/Engine/Source/Programs/AutomationTool/LiveLinkHub/LiveLinkHub.Automation.cs b/Engine/Source/Programs/AutomationTool/LiveLinkHub/LiveLinkHub.Automation.cs
--- a/Engine/Source/Programs/AutomationTool/LiveLinkHub/LiveLinkHub.Automation.cs
+++ b/Engine/Source/Programs/AutomationTool/LiveLinkHub/LiveLinkHub.Automation.cs
@@ -21,7 +21,7 @@ namespace AutomationTool
 				, NoClient: bIsCookedCooker
 				, OptionalContent: true
 				, ClientCookedTargets: new ParamList<string>() // Prevent AutodetectSettings from looking for a game target
-				, EditorTargets: new ParamList<string>("LiveLinkHubEditor")
+				, EditorTargets: new ParamList<string>("LiveLinkHub")
 				, UbtArgs: "-SingleModulePlatform"
 			);
 		}
diff --git a/Engine/Source/Programs/LiveLinkHub/Config/DefaultEditor.ini b/Engine/Source/Programs/LiveLinkHub/Config/DefaultEditor.ini
new file mode 100644
--- /dev/null
+++ b/Engine/Source/Programs/LiveLinkHub/Config/DefaultEditor.ini
@@ -0,0 +1,14 @@
+[EditorSettings]
+; We don't allow creating new plugins.
+bCanCreatePluginsFromBrowser=false
+; We don't allow toggling individual plugins.
+bCanModifyPluginsFromBrowser=false
+
+; We allow modifying user plugin paths...
+bCanModifyUserPluginDirectoriesFromBrowser=true
+; ...but not paths stored in the .uproject descriptor.
+bCanModifyProjectPluginDirectoriesFromBrowser=false
+
+; Hide plugins that aren't "Installed" or "External" to cut down on noise.
++HidePluginCategoriesFromBrowser=All
++HidePluginCategoriesFromBrowser=Built-In
diff --git a/Engine/Source/Programs/LiveLinkHub/Config/DefaultEditorPerProjectUserSettings.ini b/Engine/Source/Programs/LiveLinkHub/Config/DefaultEditorPerProjectUserSettings.ini
new file mode 100644
--- /dev/null
+++ b/Engine/Source/Programs/LiveLinkHub/Config/DefaultEditorPerProjectUserSettings.ini
@@ -0,0 +1,2 @@
+[/Script/UnrealEd.EditorLoadingSavingSettings]
+bAutoSaveEnable=False
diff --git a/Engine/Source/Programs/LiveLinkHub/Config/DefaultEngine.ini b/Engine/Source/Programs/LiveLinkHub/Config/DefaultEngine.ini
new file mode 100644
--- /dev/null
+++ b/Engine/Source/Programs/LiveLinkHub/Config/DefaultEngine.ini
@@ -0,0 +1,42 @@
+[LiveLink]
+bUseLiveLinkHubSettingsDetailCustomization=true
+bUseLiveLinkHubSubjectSettingsDetailsCustomization=true
+bPanelControllerSeparateSourcesSubjects=true
+bUseModularClientReference=true
+bEnableParentSubjects=true
+
+[LiveLinkHub]
+bEnableLLHMessageBusSourceFactory=false
+bUseMessagingConnectionManager=false
+
+[/Script/MacTargetPlatform.XcodeProjectSettings]
+TemplateMacPlist=(FilePath="/Engine/Source/Programs/LiveLinkHub/Resources/Mac/Info.plist")
+
+[/Script/LiveLinkInterface.LiveLinkDefaultSourceSettings]
+DefaultSourceFrameBufferSize=1
+
+[/Script/LiveLinkInterface.LiveLinkDefaultSubjectSettings]
+bRebroadcastSubjectsByDefault=true
+bAllowEditingRebroadcastProperty=false
+
+[/Script/LiveLinkHub.LiveLinkHubSettings]
+bAutoAddDiscoveredClients=true
++ProjectSettingsToDisplay=LiveLink
++ProjectSettingsToDisplay=LiveLinkHub
++ProjectSettingsToDisplay=UDPMessaging
+
+[/Script/AppleARKit.AppleARKitSettings]
+bEnableLiveLinkForFaceTracking=true
+
+[/Script/UnrealEd.ProjectPackagingSettings]
+bShareMaterialShaderCode=false
+
+[ConsoleVariables]
+DetailsPanel.Style.EnableCardLayout=true
+
+[/Script/BuildSettings.BuildSettings]
+DefaultEditorTarget=LiveLinkHub
+
+; Not needed at the moment, can be removed if we enable the TextureFormatOodle plugin
+[AlternateTextureCompression]
+TextureCompressionFormat=""
diff --git a/Engine/Source/Programs/LiveLinkHub/Config/DefaultGame.ini b/Engine/Source/Programs/LiveLinkHub/Config/DefaultGame.ini
new file mode 100644
--- /dev/null
+++ b/Engine/Source/Programs/LiveLinkHub/Config/DefaultGame.ini
@@ -0,0 +1,14 @@
+[CookedEditorSettings_CookedEditor]
+CookedEditorTargetName=LiveLinkHub
+bStageShaderDirs=true
+bStagePlatformBuildDirs=false
+bStageExtrasDirs=false
+bStagePlatformDirs=false
+bStageUAT=false
+bStagePython=false
+bIsForExternalDistribution=true
+
+[/Script/LiveLink.LiveLinkSettings]
+DefaultSettingsClass=/Script/LiveLinkHub.LiveLinkHubSubjectSettings
+bPreProcessRebroadcastFrames=true
+bTranslateRebroadcastFrames=true
diff --git a/Engine/Source/Programs/LiveLinkHub/LiveLinkHub.Target.cs b/Engine/Source/Programs/LiveLinkHub/LiveLinkHub.Target.cs
deleted file mode 100644
diff --git a/Engine/Source/Programs/LiveLinkHub/LiveLinkHub.uproject b/Engine/Source/Programs/LiveLinkHub/LiveLinkHub.uproject
new file mode 100644
--- /dev/null
+++ b/Engine/Source/Programs/LiveLinkHub/LiveLinkHub.uproject
@@ -0,0 +1,12 @@
+{
+	"FileVersion": 3,
+	"EngineAssociation": "",
+	"Category": "",
+	"Description": "",
+	"TargetPlatforms": [
+		"Linux",
+		"Mac",
+		"Windows"
+	],
+	"DisableEnginePluginsByDefault": true
+}
diff --git a/Engine/Source/Programs/LiveLinkHub/LiveLinkHubLauncher.Build.cs b/Engine/Source/Programs/LiveLinkHub/LiveLinkHubLauncher.Build.cs
deleted file mode 100644
diff --git a/Engine/Source/Programs/LiveLinkHub/Private/LiveLinkHubCommandlet.h b/Engine/Source/Programs/LiveLinkHub/Private/LiveLinkHubCommandlet.h
deleted file mode 100644
diff --git a/Engine/Source/Programs/LiveLinkHub/Private/LiveLinkHubPluginHelpers.cpp b/Engine/Source/Programs/LiveLinkHub/Private/LiveLinkHubPluginHelpers.cpp
deleted file mode 100644
diff --git a/Engine/Source/Programs/LiveLinkHub/Private/LiveLinkHubPluginHelpers.h b/Engine/Source/Programs/LiveLinkHub/Private/LiveLinkHubPluginHelpers.h
deleted file mode 100644
diff --git a/Engine/Source/Programs/LiveLinkHub/Private/LiveLinkHubRun.cpp b/Engine/Source/Programs/LiveLinkHub/Private/LiveLinkHubRun.cpp
deleted file mode 100644
diff --git a/Engine/Source/Programs/LiveLinkHub/Private/LiveLinkHubRun.h b/Engine/Source/Programs/LiveLinkHub/Private/LiveLinkHubRun.h
deleted file mode 100644
diff --git a/Engine/Source/Programs/LiveLinkHub/Private/LiveLinkHubRunMain.cpp b/Engine/Source/Programs/LiveLinkHub/Private/LiveLinkHubRunMain.cpp
deleted file mode 100644
diff --git a/Engine/Source/Programs/LiveLinkHub/Private/Platform/Linux/LinuxLiveLinkHubLauncherMain.cpp b/Engine/Source/Programs/LiveLinkHub/Private/Platform/Linux/LinuxLiveLinkHubLauncherMain.cpp
deleted file mode 100644
diff --git a/Engine/Source/Programs/LiveLinkHub/Private/Platform/Mac/MacLiveLinkHubLauncherMain.cpp b/Engine/Source/Programs/LiveLinkHub/Private/Platform/Mac/MacLiveLinkHubLauncherMain.cpp
deleted file mode 100644
diff --git a/Engine/Source/Programs/LiveLinkHub/Private/Platform/Windows/WindowsLiveLinkHubLauncherMain.cpp b/Engine/Source/Programs/LiveLinkHub/Private/Platform/Windows/WindowsLiveLinkHubLauncherMain.cpp
deleted file mode 100644
diff --git a/Engine/Source/Programs/LiveLinkHub/RuntimeDependencies.ini b/Engine/Source/Programs/LiveLinkHub/RuntimeDependencies.ini
deleted file mode 100644
diff --git a/Engine/Source/Programs/LiveLinkHub/Source/LiveLinkHub.Target.cs b/Engine/Source/Programs/LiveLinkHub/Source/LiveLinkHub.Target.cs
new file mode 100644
--- /dev/null
+++ b/Engine/Source/Programs/LiveLinkHub/Source/LiveLinkHub.Target.cs
@@ -0,0 +1,129 @@
+// Copyright Epic Games, Inc. All Rights Reserved.
+
+using EpicGames.Core;
+using UnrealBuildBase;
+using UnrealBuildTool;
+
+
+[SupportedPlatforms(UnrealPlatformClass.Desktop)]
+public class LiveLinkHubTarget : TargetRules
+{
+	// Restrict OptedInModulePlatforms to the current Target.Platform.
+	// Used during staging, which otherwise fails in TargetPlatform-related restricted
+	// subdirectories (Engine/Binaries/Win64/{Android,IOS,Linux,LinuxArm64,...}).
+	[CommandLine("-SingleModulePlatform")]
+	public bool bSingleModulePlatform = false;
+
+	// Whether to enable building Capture Manager plugin.
+	[CommandLine("-EnableCaptureManagerPlugin=")]
+	public bool bEnableCaptureManagerPlugin = false;
+
+	// Whether to enable building MetaHuman LiveLink plugin.
+	[CommandLine("-EnableMetaHumanLiveLinkPlugin=")]
+	public bool bEnableMetaHumanLiveLinkPlugin = false;
+
+	public LiveLinkHubTarget(TargetInfo Target) : base(Target)
+	{
+		Type = TargetType.Editor;
+		bExplicitTargetForType = true;
+		bGenerateProgramProject = true;
+
+		SolutionDirectory = "Programs/LiveLink";
+		LaunchModuleName = "LiveLinkHubLauncher";
+
+		DefaultBuildSettings = BuildSettingsVersion.Latest;
+		IncludeOrderVersion = EngineIncludeOrderVersion.Latest;
+
+		// These plugins are required for running LiveLinkHub. 
+		// They may be a direct dependency or a dependency of another plugin.
+		EnablePlugins.AddRange(new string[]
+		{
+			"ACLPlugin",
+			"LiveLink",
+			"LiveLinkHub",
+			"PluginBrowser",
+
+			"LiveLinkCamera",
+			"LiveLinkLens",
+			"LensComponent",
+			"LiveLinkInputDevice",
+			"ContentBrowserAssetDataSource",
+			"ProceduralMeshComponent",
+			"PropertyAccessEditor",
+			"PythonScriptPlugin",
+			"StructUtils",
+			"UdpMessaging",
+			"CameraCalibrationCore",
+			"AppleARKitFaceSupport",
+			"XInputDevice",
+
+			"LiveLinkOpenVR",
+		});
+
+		if (bEnableCaptureManagerPlugin)
+		{
+			OptionalPlugins.AddRange(new string[]
+			{
+				"CaptureManagerApp",
+				"CaptureManagerCore"
+			});
+		}
+
+		if (bEnableMetaHumanLiveLinkPlugin)
+		{
+			EnablePlugins.AddRange(new string[]
+			{
+				"MetaHumanLiveLink",
+				"WmfMedia"
+			});
+		}
+
+		if (bSingleModulePlatform)
+		{
+			// Necessary for staging, but avoided otherwise because it dirties
+			// Definitions.CookedEditor.h and triggers rebuilds (incl. UnrealEditor).
+			OptedInModulePlatforms = new UnrealTargetPlatform[] { Target.Platform };
+		}
+
+		bAllowEnginePluginsEnabledByDefault = false;
+		bBuildAdditionalConsoleApp = false;
+
+		// Based loosely on the VCProject.cs logic to construct NMakePath.
+		string BaseExeName = "LiveLinkHub";
+		OutputFile = "Binaries/" + Platform.ToString() + "/" + BaseExeName;
+		if (Configuration != UndecoratedConfiguration)
+		{
+			OutputFile += "-" + Platform.ToString() + "-" + Configuration.ToString();
+		}
+		if (Platform == UnrealTargetPlatform.Win64)
+		{
+			OutputFile += ".exe";
+		}
+
+		// Copy the target receipt into the project binaries directory.
+		// Prevents "Would you like to build the editor?" prompt on startup when running with project context.
+		//
+		// TODO?: If instead we moved the .uproject out of Engine/Source/Programs and into Engine/Programs, UBT
+		// would correctly(?) pass `-Project=` in the build command. However, in addition to the .target file,
+		// this also results in the executable and a _subset_ of the dependent DLLs ending up ending up in the
+		// project Binaries directory. So for now, we'll keep splitting the difference like this.
+		DirectoryReference ReceiptSrcDir = Unreal.EngineDirectory;
+		DirectoryReference ReceiptDestDir = DirectoryReference.Combine(
+			Unreal.EngineDirectory, "Source", "Programs", "LiveLinkHub");
+
+		FileReference ReceiptSrcPath = TargetReceipt.GetDefaultPath(ReceiptSrcDir, BaseExeName, Platform, Configuration, Architectures);
+		FileReference ReceiptDestPath = TargetReceipt.GetDefaultPath(ReceiptDestDir, BaseExeName, Platform, Configuration, Architectures);
+		DirectoryReference.CreateDirectory(ReceiptDestPath.Directory);
+
+		PostBuildSteps.Add($"echo Copying \"{ReceiptSrcPath}\" to \"{ReceiptDestPath}\"");
+
+		if (Platform == UnrealTargetPlatform.Win64)
+		{
+			PostBuildSteps.Add($"copy /Y \"{ReceiptSrcPath}\" \"{ReceiptDestPath}\"");
+		}
+		else
+		{
+			PostBuildSteps.Add($"cp -a \"{ReceiptSrcPath}\" \"{ReceiptDestPath}\"");
+		}
+	}
+}
diff --git a/Engine/Source/Programs/LiveLinkHub/Source/LiveLinkHubLauncher/LiveLinkHubLauncher.Build.cs b/Engine/Source/Programs/LiveLinkHub/Source/LiveLinkHubLauncher/LiveLinkHubLauncher.Build.cs
new file mode 100644
--- /dev/null
+++ b/Engine/Source/Programs/LiveLinkHub/Source/LiveLinkHubLauncher/LiveLinkHubLauncher.Build.cs
@@ -0,0 +1,139 @@
+// Copyright Epic Games, Inc. All Rights Reserved.
+
+#nullable enable
+
+using System;
+using System.Collections.Generic;
+using System.Linq;
+using EpicGames.Core;
+using UnrealBuildTool;
+
+
+// This is named LiveLinkHubLauncher in order to not clash with the LiveLinkHub module which is in the livelink plugin.
+public class LiveLinkHubLauncher : ModuleRules
+{
+	public LiveLinkHubLauncher(ReadOnlyTargetRules Target) : base(Target)
+	{
+		PrivateDependencyModuleNames.AddRange(
+			new string[] {
+				"LiveLinkHub",
+				"PluginBrowser",
+			});
+
+		// LaunchEngineLoop dependencies
+		PrivateDependencyModuleNames.AddRange(
+			new string[] {
+				"ApplicationCore",
+				"AssetRegistry",
+				"Core",
+				"CoreUObject",
+				"DesktopPlatform",
+				"Engine",
+				"InputCore",
+				"InstallBundleManager",
+				"MediaUtils",
+				"Messaging",
+				"MoviePlayer",
+				"MoviePlayerProxy",
+				"ProfileVisualizer",
+				"ProfilerService",
+				"Projects",
+				"PreLoadScreen",
+				"PIEPreviewDeviceProfileSelector",
+				"RenderCore",
+				"RHI",
+				"Slate",
+				"SlateCore",
+				"StandaloneRenderer",
+				"TraceLog", 
+			}
+		);
+
+		// LaunchEngineLoop IncludePath dependencies
+		PrivateIncludePathModuleNames.AddRange(
+			new string[] {
+				"Launch",
+				"AutomationWorker",
+				"AutomationController",
+				"AutomationTest",
+				"DerivedDataCache",
+				"HeadMountedDisplay", 
+				"MRMesh", 
+				"SlateRHIRenderer", 
+				"SlateNullRenderer",
+			}
+		);
+
+		// LaunchEngineLoop editor dependencies
+		PrivateDependencyModuleNames.AddRange(
+			new string[] {
+				"PropertyEditor",
+				"DerivedDataCache",
+				"ToolWidgets",
+				"UnrealEd"
+		});
+
+
+		if (Target.Configuration != UnrealTargetConfiguration.Shipping)
+		{
+			PrivateDependencyModuleNames.AddRange(
+				new string[] {
+					"AutomationController",
+					"AutomationTest",
+					"AutomationWorker",
+				});
+		}
+
+		if (Target.IsInPlatformGroup(UnrealPlatformGroup.Linux))
+		{
+			PrivateDependencyModuleNames.AddRange(
+				new string[] {
+					"UnixCommonStartup"
+				}
+			);
+		}
+
+		ShortName = "LLHub";
+
+		AddRuntimeDependencies();
+	}
+
+	void AddRuntimeDependencies()
+	{
+		DirectoryReference ModuleDirectoryRef = new DirectoryReference(ModuleDirectory);
+		FileReference ConfigLocation = FileReference.Combine(ModuleDirectoryRef, "RuntimeDependencies.ini");
+		ConfigHierarchy Config = new ConfigHierarchy(new List<ConfigFile>{ new ConfigFile(ConfigLocation) });
+
+		if (!Config.TryGetValues("LiveLinkHubLauncher", "RuntimeDependencies", out IReadOnlyList<string>? ConfigDependencies))
+		{
+			string Err = "Error reading runtime dependencies from config";
+			Console.WriteLine(Err);
+			throw new BuildException(Err);
+		}
+
+		FileFilter Filter = new FileFilter();
+		foreach (string Rule in ConfigDependencies)
+		{
+			Filter.AddRule(Rule);
+		}
+
+		// Apply the standard exclusion rules (c.f. DeploymentContext ctor)
+		HashSet<string> RestrictedFolderNames = new HashSet<string>(StringComparer.OrdinalIgnoreCase);
+		RestrictedFolderNames.UnionWith(PlatformExports.GetPlatformFolderNames());
+		RestrictedFolderNames.UnionWith(RestrictedFolder.GetNames());
+		RestrictedFolderNames.ExceptWith(PlatformExports.GetIncludedFolderNames(Target.Platform));
+		foreach (string RestrictedFolderName in RestrictedFolderNames)
+		{
+			string Rule = String.Format(".../{0}/...", RestrictedFolderName);
+			Filter.AddRule(Rule, FileFilterType.Exclude);
+		}
+
+		// Make relative to $(EngineDir) and add to RuntimeDependencies
+		List<FileReference> AbsFiles = Filter.ApplyToDirectory(new DirectoryReference(EngineDirectory), true);
+		IEnumerable<string> RelFiles = AbsFiles.Select(x => x.FullName.Replace(EngineDirectory, "$(EngineDir)"));
+		foreach (string RelFile in RelFiles)
+		{
+			RuntimeDependencies.Add(RelFile, StagedFileType.UFS);
+		}
+	}
+}
diff --git a/Engine/Source/Programs/LiveLinkHub/Source/LiveLinkHubLauncher/Private/LiveLinkHubCommandlet.h b/Engine/Source/Programs/LiveLinkHub/Source/LiveLinkHubLauncher/Private/LiveLinkHubCommandlet.h
new file mode 100644
--- /dev/null
+++ b/Engine/Source/Programs/LiveLinkHub/Source/LiveLinkHubLauncher/Private/LiveLinkHubCommandlet.h
@@ -0,0 +1,49 @@
+// Copyright Epic Games, Inc. All Rights Reserved.
+
+#pragma once
+
+#include "Commandlets/Commandlet.h"
+#include "Misc/CoreDelegates.h"
+#include "Modules/ModuleManager.h"
+#include "ILiveLinkHubModule.h"
+
+#include "LiveLinkHubCommandlet.generated.h"
+
+UCLASS()
+class ULiveLinkHubCommandlet : public UCommandlet
+{
+public:
+	GENERATED_BODY()
+
+	ULiveLinkHubCommandlet()
+	{
+		LogToConsole = false;
+		FastExit = true;
+	}
+	
+	/** Runs the commandlet */
+	virtual int32 Main(const FString& Params) override
+	{
+		PRIVATE_GAllowCommandletRendering = true;
+
+		// Work around the CB file asset data source being disabled in commandlets.
+		if (IModuleInterface* AssetDataSourceModule = FModuleManager::Get().GetModule("ContentBrowserAssetDataSource"))
+		{
+			PRIVATE_GIsRunningCommandlet = false;
+			AssetDataSourceModule->StartupModule();
+			PRIVATE_GIsRunningCommandlet = true;
+		}
+
+		FModuleManager::Get().LoadModule("OutputLog");
+		FModuleManager::Get().LoadModule("ProjectSettingsViewer");
+
+		FModuleManager::Get().LoadModuleChecked<ILiveLinkHubModule>("LiveLinkEditor");
+
+		ILiveLinkHubModule& HubModule = FModuleManager::Get().LoadModuleChecked<ILiveLinkHubModule>("LiveLinkHub");
+		HubModule.PreinitializeLiveLinkHub();
+		HubModule.StartLiveLinkHub();
+		HubModule.ShutdownLiveLinkHub();
+
+		return 0;
+	}
+};
diff --git a/Engine/Source/Programs/LiveLinkHub/Source/LiveLinkHubLauncher/Private/LiveLinkHubPluginHelpers.cpp b/Engine/Source/Programs/LiveLinkHub/Source/LiveLinkHubLauncher/Private/LiveLinkHubPluginHelpers.cpp
new file mode 100644
--- /dev/null
+++ b/Engine/Source/Programs/LiveLinkHub/Source/LiveLinkHubLauncher/Private/LiveLinkHubPluginHelpers.cpp
@@ -0,0 +1,93 @@
+// Copyright Epic Games, Inc. All Rights Reserved.
+
+#include "LiveLinkHubPluginHelpers.h"
+#include "HAL/PlatformProcess.h"
+#include "Interfaces/IPluginManager.h"
+#include "IPluginBrowser.h"
+#include "Misc/App.h"
+#include "Misc/ConfigCacheIni.h"
+
+
+namespace LiveLinkHub::PluginHelpers
+{
+	const TCHAR* const ConfigSection = TEXT("LiveLinkHub");
+	const TCHAR* const ConfigKey = TEXT("PluginDirectories");
+}
+
+
+FString LiveLinkHub::PluginHelpers::GetPluginConfigPath()
+{
+	return FPaths::Combine(
+		FPlatformProcess::UserSettingsDir(),
+		*FApp::GetEpicProductIdentifier(),
+		TEXT("LiveLinkHub"),
+		TEXT("Config"),
+		TEXT("LiveLinkHubPlugins.ini")
+	);
+}
+
+void LiveLinkHub::PluginHelpers::ReadPluginDirectoriesFromConfig(TArray<FString>& OutDirectories)
+{
+	FConfigFile ConfigFile;
+	ConfigFile.Read(GetPluginConfigPath());
+	ConfigFile.GetArray(ConfigSection, ConfigKey, OutDirectories);
+}
+
+void LiveLinkHub::PluginHelpers::WritePluginDirectoriesToConfig(const TArray<FString>& InDirectories)
+{
+	FConfigFile ConfigFile;
+	ConfigFile.Read(GetPluginConfigPath());
+
+	ConfigFile.ResetKeyInSection(ConfigSection, ConfigKey);
+	for (const FString& Directory : InDirectories)
+	{
+		ConfigFile.AddUniqueToSection(ConfigSection, ConfigKey, Directory);
+	}
+
+	ConfigFile.Write(GetPluginConfigPath());
+}
+
+void LiveLinkHub::PluginHelpers::RestoreSavedPluginDirectories()
+{
+	TArray<FString> SavedDirectories;
+	ReadPluginDirectoriesFromConfig(SavedDirectories);
+
+	bool bPluginPathsChanged = false;
+
+	IPluginManager& PluginManager = IPluginManager::Get();
+	for (const FString& Directory : SavedDirectories)
+	{
+		const bool bRefresh_false = false;
+		bPluginPathsChanged |= PluginManager.AddPluginSearchPath(Directory, bRefresh_false);
+	}
+
+	if (bPluginPathsChanged)
+	{
+		PluginManager.RefreshPluginsList();
+	}
+}
+
+void LiveLinkHub::PluginHelpers::HandlePluginDirectoriesChanged()
+{
+	IPluginManager& PluginManager = IPluginManager::Get();
+
+	TSet<FExternalPluginPath> PluginSources;
+	PluginManager.GetExternalPluginSources(PluginSources);
+
+	TArray<FString> UserDirectories;
+	for (const FExternalPluginPath& Path : PluginSources)
+	{
+		if (Path.Source == EPluginExternalSource::Other)
+		{
+			UserDirectories.Add(Path.Path);
+		}
+	}
+
+	WritePluginDirectoriesToConfig(UserDirectories);
+}
+
+void LiveLinkHub::PluginHelpers::RegisterPluginDirectoriesChangedHandler()
+{
+	IPluginBrowser& PluginBrowser = IPluginBrowser::Get();
+	PluginBrowser.OnPluginDirectoriesChanged().AddStatic(HandlePluginDirectoriesChanged);
+}
diff --git a/Engine/Source/Programs/LiveLinkHub/Source/LiveLinkHubLauncher/Private/LiveLinkHubPluginHelpers.h b/Engine/Source/Programs/LiveLinkHub/Source/LiveLinkHubLauncher/Private/LiveLinkHubPluginHelpers.h
new file mode 100644
--- /dev/null
+++ b/Engine/Source/Programs/LiveLinkHub/Source/LiveLinkHubLauncher/Private/LiveLinkHubPluginHelpers.h
@@ -0,0 +1,26 @@
+// Copyright Epic Games, Inc. All Rights Reserved.
+
+#include "Containers/Array.h"
+#include "Containers/UnrealString.h"
+
+
+namespace LiveLinkHub::PluginHelpers
+{
+	/** Returns the path to LiveLinkHubPlugins.ini. */
+	FString GetPluginConfigPath();
+
+	/** Returns the user plugin directories stored in LiveLinkHubPlugins.ini. */
+	void ReadPluginDirectoriesFromConfig(TArray<FString>& OutDirectories);
+
+	/** Serializes the specified user plugin directories out to LiveLinkHubPlugins.ini. */
+	void WritePluginDirectoriesToConfig(const TArray<FString>& InDirectories);
+
+	/** Reads the configured user plugin directories and adds them to the plugin manager search path. */
+	void RestoreSavedPluginDirectories();
+
+	/** Stores the latest plugin manager search paths to disk. */
+	void HandlePluginDirectoriesChanged();
+
+	/** Connects the `HandlePluginDirectoriesChanged` handler to the plugin browser directories changed event. */
+	void RegisterPluginDirectoriesChangedHandler();
+}
diff --git a/Engine/Source/Programs/LiveLinkHub/Source/LiveLinkHubLauncher/Private/LiveLinkHubRun.cpp b/Engine/Source/Programs/LiveLinkHub/Source/LiveLinkHubLauncher/Private/LiveLinkHubRun.cpp
new file mode 100644
--- /dev/null
+++ b/Engine/Source/Programs/LiveLinkHub/Source/LiveLinkHubLauncher/Private/LiveLinkHubRun.cpp
@@ -0,0 +1,146 @@
+// Copyright Epic Games, Inc. All Rights Reserved.
+
+#include "LiveLinkHubRun.h"
+
+#include "Editor.h"
+#include "Editor/EditorEngine.h"
+#include "Misc/App.h"
+#include "Misc/PackageName.h"
+#include "Misc/Paths.h"
+#include "HAL/FileManager.h"
+#include "HAL/PlatformApplicationMisc.h"
+#include "HAL/PlatformFileManager.h"
+#include "HAL/PlatformMisc.h"
+#include "HAL/PlatformSplash.h"
+#include "ILiveLinkHubModule.h"
+#include "LaunchEngineLoop.h"
+#include "LiveLinkHubPluginHelpers.h"
+#include "Misc/CommandLine.h"
+#include "Misc/CoreDelegates.h"
+#include "Misc/Paths.h"
+#include "Modules/ModuleManager.h"
+
+#ifndef WITH_ASSET_LOADING_AUDIT
+#define WITH_ASSET_LOADING_AUDIT 0
+#endif
+
+
+DEFINE_LOG_CATEGORY_STATIC(LogLiveLinkHub, Log, All);
+
+
+int32 RunLiveLinkHub(const TCHAR* CommandLine)
+{
+	FTaskTagScope Scope(ETaskTag::EGameThread);
+
+	// Needs to be initialized early for splash / mount points / plugin search paths.
+	FCommandLine::Set(CommandLine);
+
+	const FText AppName = NSLOCTEXT("LiveLinkHub", "SplashTextName", "Live Link Hub");
+	FPlatformSplash::SetSplashText(SplashTextType::GameName, *AppName.ToString());
+
+#if !UE_BUILD_SHIPPING
+	if (FParse::Param(CommandLine, TEXT("WaitForDebugger")))
+	{
+		while (!FPlatformMisc::IsDebuggerPresent())
+		{
+			FPlatformProcess::Sleep(0.1f);
+		}
+		UE_DEBUG_BREAK();
+	}
+#endif
+
+#if WITH_ASSET_LOADING_AUDIT
+	FCoreDelegates::OnSyncLoadPackage.AddLambda([](const FString& PackageName)
+		{
+			UE_LOG(LogLiveLinkHub, Display, TEXT("Audit: Loaded %s"), *PackageName);
+		});
+#endif
+
+	bool bLauncherDistribution = false;
+
+#if !IS_PROGRAM
+	const TCHAR* const DevelopmentProjectPath = TEXT("../../Source/Programs/LiveLinkHub/LiveLinkHub.uproject");
+	const TCHAR* const StagedProjectPath = TEXT("../../../LiveLinkHub/LiveLinkHub.uproject");
+
+	if (FPaths::FileExists(DevelopmentProjectPath))
+	{
+		FPaths::SetProjectFilePath(DevelopmentProjectPath);
+	}
+	else if (FPaths::FileExists(StagedProjectPath))
+	{
+		bLauncherDistribution = true;
+		FPaths::SetProjectFilePath(StagedProjectPath);
+	}
+
+	if (bLauncherDistribution)
+	{
+		// This is needed to ensure that content generated by the hub will be discoverable (ie. LiveLink recordings)
+		const FString GameContentPath = FPaths::Combine(FPlatformProcess::UserSettingsDir(), *FApp::GetEpicProductIdentifier(), TEXT("LiveLinkHub"), TEXT("Content"));
+		FPackageName::RegisterMountPoint(TEXT("/Game/"), GameContentPath);
+	}
+
+	// Used by Live Coding, among other things.
+	FPlatformMisc::SetUBTTargetName(TEXT("LiveLinkHub"));
+
+	// We need to specify this manually, because it isn't inferred from the command line.
+	// But we can't set it too early, because LaunchSetGameName clears it.
+	FCoreDelegates::OnInit.AddStatic([]()
+	{
+		// Required for OpenXR instance registration.
+		FApp::SetProjectName(TEXT("LiveLinkHub"));
+	});
+
+	// Plugin directory config save/load
+	FCoreDelegates::TSConfigReadyForUse().AddStatic(LiveLinkHub::PluginHelpers::RestoreSavedPluginDirectories);
+	FCoreDelegates::OnAllModuleLoadingPhasesComplete.AddStatic(LiveLinkHub::PluginHelpers::RegisterPluginDirectoriesChangedHandler);
+#endif
+
+#if IS_PROGRAM
+	// Disable this when going through PreInit to prevent the console window from appearing.
+	GIsSilent = true;
+#endif
+
+	// Start up the main loop, adding some extra command line arguments:
+#if !IS_PROGRAM
+	int32 Result = GEngineLoop.PreInit(*FString::Printf(TEXT("%s %s"), CommandLine, TEXT("-xrtrackingonly")));
+#else
+	int32 Result = GEngineLoop.PreInit(*FString::Printf(TEXT("%s %s"), CommandLine, TEXT("-RUN=LiveLinkHubCommandlet -Messaging -DDC=NoShared -NoShaderCompile")));
+#endif
+
+	if (Result != 0)
+	{
+		UE_LOG(LogLiveLinkHub, Error, TEXT("EngineLoop PreInit failed (%i)"), Result);
+	}
+
+	if (Result != 0 || IsEngineExitRequested())
+	{
+		return Result;
+	}
+
+#if !IS_PROGRAM
+	{
+		ILiveLinkHubModule& HubModule = FModuleManager::Get().LoadModuleChecked<ILiveLinkHubModule>("LiveLinkHub");
+		HubModule.PreinitializeLiveLinkHub();
+
+		Result = GEngineLoop.Init();
+
+		GEditor->ShouldDisableCPUThrottlingDelegates.Add(UEditorEngine::FShouldDisableCPUThrottling::CreateLambda([]() { return true; }));
+
+		HubModule.StartLiveLinkHub(bLauncherDistribution);
+
+		// Hide the splash screen now that everything is ready to go
+		FPlatformSplash::Hide();
+
+		while (!IsEngineExitRequested())
+		{
+			GEngineLoop.Tick();
+		}
+
+		HubModule.ShutdownLiveLinkHub();
+	}
+#endif
+
+	GEngineLoop.Exit();
+
+	return Result;
+}
diff --git a/Engine/Source/Programs/LiveLinkHub/Source/LiveLinkHubLauncher/Private/LiveLinkHubRun.h b/Engine/Source/Programs/LiveLinkHub/Source/LiveLinkHubLauncher/Private/LiveLinkHubRun.h
new file mode 100644
--- /dev/null
+++ b/Engine/Source/Programs/LiveLinkHub/Source/LiveLinkHubLauncher/Private/LiveLinkHubRun.h
@@ -0,0 +1,7 @@
+// Copyright Epic Games, Inc. All Rights Reserved.
+
+#pragma once
+
+#include "HAL/Platform.h"
+
+int32 RunLiveLinkHub(const TCHAR* CommandLine);
diff --git a/Engine/Source/Programs/LiveLinkHub/Source/LiveLinkHubLauncher/Private/LiveLinkHubRunMain.cpp b/Engine/Source/Programs/LiveLinkHub/Source/LiveLinkHubLauncher/Private/LiveLinkHubRunMain.cpp
new file mode 100644
--- /dev/null
+++ b/Engine/Source/Programs/LiveLinkHub/Source/LiveLinkHubLauncher/Private/LiveLinkHubRunMain.cpp
@@ -0,0 +1,10 @@
+// Copyright Epic Games, Inc. All Rights Reserved.
+
+#include "HAL/Platform.h"
+#include "RequiredProgramMainCPPInclude.h"
+
+#if IS_PROGRAM
+IMPLEMENT_APPLICATION(LiveLinkHubLauncher, "LiveLinkHub");
+#else
+FEngineLoop GEngineLoop;
+#endif
diff --git a/Engine/Source/Programs/LiveLinkHub/Source/LiveLinkHubLauncher/Private/Platform/Linux/LinuxLiveLinkHubLauncherMain.cpp b/Engine/Source/Programs/LiveLinkHub/Source/LiveLinkHubLauncher/Private/Platform/Linux/LinuxLiveLinkHubLauncherMain.cpp
new file mode 100644
--- /dev/null
+++ b/Engine/Source/Programs/LiveLinkHub/Source/LiveLinkHubLauncher/Private/Platform/Linux/LinuxLiveLinkHubLauncherMain.cpp
@@ -0,0 +1,11 @@
+// Copyright Epic Games, Inc. All Rights Reserved.
+
+#include "CoreMinimal.h"
+#include "LiveLinkHubRun.h"
+#include "LaunchEngineLoop.h"
+#include "UnixCommonStartup.h"
+
+int main(int argc, char *argv[])
+{
+	return CommonUnixMain(argc, argv, &RunLiveLinkHub, [] { } );
+}
diff --git a/Engine/Source/Programs/LiveLinkHub/Source/LiveLinkHubLauncher/Private/Platform/Mac/MacLiveLinkHubLauncherMain.cpp b/Engine/Source/Programs/LiveLinkHub/Source/LiveLinkHubLauncher/Private/Platform/Mac/MacLiveLinkHubLauncherMain.cpp
new file mode 100644
--- /dev/null
+++ b/Engine/Source/Programs/LiveLinkHub/Source/LiveLinkHubLauncher/Private/Platform/Mac/MacLiveLinkHubLauncherMain.cpp
@@ -0,0 +1,10 @@
+// Copyright Epic Games, Inc. All Rights Reserved.
+
+#include "LiveLinkHubRun.h"
+#include "Mac/MacProgramDelegate.h"
+#include "LaunchEngineLoop.h"
+
+int main( int argc, char *argv[] )
+{
+	return [MacProgramDelegate mainWithArgc:argc argv:argv programMain:RunLiveLinkHub programExit:FEngineLoop::AppExit];
+}
diff --git a/Engine/Source/Programs/LiveLinkHub/Source/LiveLinkHubLauncher/Private/Platform/Windows/WindowsLiveLinkHubLauncherMain.cpp b/Engine/Source/Programs/LiveLinkHub/Source/LiveLinkHubLauncher/Private/Platform/Windows/WindowsLiveLinkHubLauncherMain.cpp
new file mode 100644
--- /dev/null
+++ b/Engine/Source/Programs/LiveLinkHub/Source/LiveLinkHubLauncher/Private/Platform/Windows/WindowsLiveLinkHubLauncherMain.cpp
@@ -0,0 +1,79 @@
+// Copyright Epic Games, Inc. All Rights Reserved.
+
+#include "CoreMinimal.h"
+#include "LiveLinkHubRun.h"
+#include "HAL/ExceptionHandling.h"
+#include "LaunchEngineLoop.h"
+#include "Windows/WindowsHWrapper.h"
+#include "Misc/CommandLine.h"
+#include "Misc/OutputDeviceError.h"
+
+
+// Opt in to new D3D12 redist and tell the loader where to search for D3D12Core.dll.
+// The D3D loader looks for these symbol exports in the .exe module.
+// We only support this on x64 Windows Desktop platforms. Other platforms or non-redist-aware 
+// versions of Windows will transparently load default OS-provided D3D12 library.
+#define USE_D3D12_REDIST (PLATFORM_DESKTOP && PLATFORM_CPU_X86_FAMILY && PLATFORM_64BITS && 1)
+#if USE_D3D12_REDIST
+extern "C" { _declspec(dllexport) extern const UINT D3D12SDKVersion = 614; } // D3D12_SDK_VERSION
+extern "C" { _declspec(dllexport) extern const char* D3D12SDKPath = ".\\D3D12\\"; }
+#endif // USE_D3D12_REDIST
+
+
+/**
+ * The main application entry point for Windows platforms.
+ *
+ * @param hInInstance Handle to the current instance of the application.
+ * @param hPrevInstance Handle to the previous instance of the application (always NULL).
+ * @param lpCmdLine Command line for the application.
+ * @param nShowCmd Specifies how the window is to be shown.
+ * @return Application's exit value.
+ */
+int32 WINAPI WinMain(_In_ HINSTANCE hInInstance, _In_opt_ HINSTANCE hPrevInstance, _In_ char* lpCmdLine, _In_ int32 nShowCmd)
+{
+	hInstance = hInInstance;
+
+	const TCHAR* CmdLine = ::GetCommandLineW();
+	CmdLine = FCommandLine::RemoveExeName(CmdLine);
+
+#if !UE_BUILD_SHIPPING
+	if (FParse::Param(CmdLine, TEXT("crashreports")))
+	{
+		GAlwaysReportCrash = true;
+	}
+#endif
+
+	int32 ErrorLevel = 0;
+
+#if UE_BUILD_DEBUG
+	if (!GAlwaysReportCrash)
+#else
+	if (FPlatformMisc::IsDebuggerPresent() && !GAlwaysReportCrash)
+#endif
+	{
+		ErrorLevel = RunLiveLinkHub(CmdLine);
+	}
+	else
+	{
+#if !PLATFORM_SEH_EXCEPTIONS_DISABLED
+		__try
+#endif
+ 		{
+			GIsGuarded = 1;
+			ErrorLevel = RunLiveLinkHub(CmdLine);
+			GIsGuarded = 0;
+		}
+#if !PLATFORM_SEH_EXCEPTIONS_DISABLED
+		__except (ReportCrash(GetExceptionInformation()))
+		{
+			ErrorLevel = 1;
+			GError->HandleError();
+			FPlatformMisc::RequestExit(true);
+		}
+#endif
+	}
+
+	FEngineLoop::AppExit();
+
+	return ErrorLevel;
+}
diff --git a/Engine/Source/Programs/LiveLinkHub/Source/LiveLinkHubLauncher/RuntimeDependencies.ini b/Engine/Source/Programs/LiveLinkHub/Source/LiveLinkHubLauncher/RuntimeDependencies.ini
new file mode 100644
--- /dev/null
+++ b/Engine/Source/Programs/LiveLinkHub/Source/LiveLinkHubLauncher/RuntimeDependencies.ini
@@ -0,0 +1,90 @@
+[LiveLinkHubLauncher]
++RuntimeDependencies="/Content/Animation/..."
++RuntimeDependencies="/Content/BasicShapes/..."
++RuntimeDependencies="/Content/Editor/Slate/..."
+; Prefix - denotes exclusion
++RuntimeDependencies="-/Content/Editor/Slate/Automation/..."
++RuntimeDependencies="-/Content/Editor/Slate/Testing/..."
++RuntimeDependencies="/Content/EditorMaterials/..."
++RuntimeDependencies="/Content/EditorMeshes/Camera/..."
++RuntimeDependencies="/Content/EditorResources/..."
++RuntimeDependencies="/Content/EngineDamageTypes/..."
++RuntimeDependencies="/Content/EngineFonts/..."
++RuntimeDependencies="/Content/EngineDebugMaterials/..."
++RuntimeDependencies="/Content/EngineMaterials/..."
++RuntimeDependencies="/Content/Slate/..."
+
+; Individual assets
++RuntimeDependencies="/Content/ArtTools/RenderToTexture/Meshes/S_1_Unit_Plane.uasset"
++RuntimeDependencies="/Content/EditorLandscapeResources/DataLayer.uasset"
++RuntimeDependencies="/Content/EditorLandscapeResources/DefaultAlphaTexture.uasset"
++RuntimeDependencies="/Content/EditorLandscapeResources/SplineEditorMesh.uasset"
++RuntimeDependencies="/Content/EditorLandscapeResources/SplineEditorMeshMat.uasset"
++RuntimeDependencies="/Content/EngineMeshes/Sphere.uasset"
++RuntimeDependencies="/Content/EngineResources/Black.uasset"
++RuntimeDependencies="/Content/EngineResources/DefaultTexture.uasset"
++RuntimeDependencies="/Content/EngineResources/DefaultTextureCube.uasset"
++RuntimeDependencies="/Content/EngineResources/DefaultVolumeTexture.uasset"
++RuntimeDependencies="/Content/EngineResources/DefaultVolumeTexture2D.uasset"
++RuntimeDependencies="/Content/EngineResources/GradientTexture0.uasset"
++RuntimeDependencies="/Content/EngineResources/WhiteSquareTexture.uasset"
++RuntimeDependencies="/Content/EngineSky/VolumetricClouds/CloudGradientTexture.uasset"
++RuntimeDependencies="/Content/EngineSky/VolumetricClouds/CloudWeatherTexture.uasset"
++RuntimeDependencies="/Content/EngineSky/VolumetricClouds/T_NoiseErosion.uasset"
++RuntimeDependencies="/Content/EngineSky/VolumetricClouds/T_NoiseShape64.uasset"
++RuntimeDependencies="/Content/EngineSky/VolumetricClouds/T_VolumeNoiseErosion32.uasset"
++RuntimeDependencies="/Content/EngineSky/VolumetricClouds/T_VolumeNoiseShape64.uasset"
++RuntimeDependencies="/Content/EngineSky/VolumetricClouds/m_SimpleVolumetricCloud.uasset"
++RuntimeDependencies="/Content/EngineSky/VolumetricClouds/m_SimpleVolumetricCloud_Inst.uasset"
++RuntimeDependencies="/Content/Functions/Engine_MaterialFunctions01/Opacity/CameraDepthFade.uasset"
++RuntimeDependencies="/Content/Functions/Engine_MaterialFunctions01/Shading/ConvertFromDiffSpec.uasset"
++RuntimeDependencies="/Content/Functions/Engine_MaterialFunctions01/Shading/PowerToRoughness.uasset"
++RuntimeDependencies="/Content/Functions/Engine_MaterialFunctions02/Math/Transform3x3Matrix.uasset"
++RuntimeDependencies="/Content/Functions/Engine_MaterialFunctions02/PivotPainter2/Black_1x1_EXR_Texture.uasset"
++RuntimeDependencies="/Content/Functions/Engine_MaterialFunctions02/SplitComponents.uasset"
++RuntimeDependencies="/Content/Functions/Engine_MaterialFunctions02/Utility/BreakOutFloat2Components.uasset"
++RuntimeDependencies="/Content/Functions/Engine_MaterialFunctions02/Utility/BreakOutFloat3Components.uasset"
++RuntimeDependencies="/Content/Functions/Engine_MaterialFunctions02/Utility/CameraDirectionVector.uasset"
++RuntimeDependencies="/Content/Functions/Engine_MaterialFunctions02/Utility/MakeFloat2.uasset"
++RuntimeDependencies="/Content/Functions/Engine_MaterialFunctions02/Utility/MakeFloat3.uasset"
++RuntimeDependencies="/Content/Functions/Engine_MaterialFunctions02/Utility/MakeFloat4.uasset"
++RuntimeDependencies="/Content/Functions/Engine_MaterialFunctions02/WorldPositionOffset/CameraOffset.uasset"
++RuntimeDependencies="/Content/VREditor/LaserPointer/VR_LaserPower_01.uasset"
++RuntimeDependencies="/Content/VREditor/SnapGrid/SnapGridMaterial.uasset"
++RuntimeDependencies="/Content/VREditor/SnapGrid/SnapGridPlaneMesh.uasset"
++RuntimeDependencies="/Content/VREditor/Sounds/UI/Drop_Item_From_ContentBrowser.uasset"
++RuntimeDependencies="/Content/VREditor/Sounds/UI/Gizmo_Handle_Clicked.uasset"
++RuntimeDependencies="/Content/VREditor/Sounds/UI/Object_PickUp.uasset"
++RuntimeDependencies="/Content/VREditor/Sounds/UI/Object_Snaps_To_Another_Actor.uasset"
++RuntimeDependencies="/Content/VREditor/Sounds/UI/Object_Snaps_To_Grid.uasset"
++RuntimeDependencies="/Content/VREditor/Sounds/UI/Scaling_Down.uasset"
++RuntimeDependencies="/Content/VREditor/Sounds/UI/Scaling_Up.uasset"
++RuntimeDependencies="/Content/VREditor/Sounds/UI/Selection_Changes.uasset"
++RuntimeDependencies="/Content/VREditor/TransformGizmo/BoundingBoxCorner.uasset"
++RuntimeDependencies="/Content/VREditor/TransformGizmo/BoundingBoxEdge.uasset"
++RuntimeDependencies="/Content/VREditor/TransformGizmo/Main.uasset"
++RuntimeDependencies="/Content/VREditor/TransformGizmo/PlaneTranslationHandle.uasset"
++RuntimeDependencies="/Content/VREditor/TransformGizmo/RotationHandleFull.uasset"
++RuntimeDependencies="/Content/VREditor/TransformGizmo/RotationHandleIndicator.uasset"
++RuntimeDependencies="/Content/VREditor/TransformGizmo/RotationHandleQuarter.uasset"
++RuntimeDependencies="/Content/VREditor/TransformGizmo/SM_Sequencer_Node.uasset"
++RuntimeDependencies="/Content/VREditor/TransformGizmo/StartRotationHandleIndicator.uasset"
++RuntimeDependencies="/Content/VREditor/TransformGizmo/TransformGizmoFreeRotation.uasset"
++RuntimeDependencies="/Content/VREditor/TransformGizmo/TransformGizmoMaterial.uasset"
++RuntimeDependencies="/Content/VREditor/TransformGizmo/TransformGizmoMaterial_Inst.uasset"
++RuntimeDependencies="/Content/VREditor/TransformGizmo/TranslateArrowHandle.uasset"
++RuntimeDependencies="/Content/VREditor/TransformGizmo/TranslucentTransformGizmoMaterial.uasset"
++RuntimeDependencies="/Content/VREditor/TransformGizmo/TranslucentTransformGizmoMaterial_Inst.uasset"
++RuntimeDependencies="/Content/VREditor/TransformGizmo/UniformScaleHandle.uasset"
++RuntimeDependencies="/Content/VREditor/TransformGizmo/Xray.uasset"
++RuntimeDependencies="/Content/VREditor/ViewportInteractionAssetContainerData.uasset"
+
+; Plugin content (only needed for program BuildCookRun; MakeCookedEditor includes by default)
+; FIXME: ACLPlugin still generates warnings at runtime
++RuntimeDependencies="/Plugins/Animation/ACLPlugin/Content/..."
++RuntimeDependencies="/Plugins/Animation/LiveLink/Content/..."
++RuntimeDependencies="/Plugins/MovieScene/LevelSequenceEditor/Content/..."
++RuntimeDependencies="/Plugins/Runtime/AR/ARUtilities/Content/..."
++RuntimeDependencies="/Plugins/Runtime/AR/AppleAR/AppleARKit/Content/..."
++RuntimeDependencies="/Plugins/VirtualProduction/CameraCalibrationCore/Content/..."
++RuntimeDependencies="/Plugins/VirtualProduction/Takes/Content/..."
diff --git a/Engine/Source/Programs/LiveLinkHubEditor/Config/DefaultEditor.ini b/Engine/Source/Programs/LiveLinkHubEditor/Config/DefaultEditor.ini
deleted file mode 100644
diff --git a/Engine/Source/Programs/LiveLinkHubEditor/Config/DefaultEditorPerProjectUserSettings.ini b/Engine/Source/Programs/LiveLinkHubEditor/Config/DefaultEditorPerProjectUserSettings.ini
deleted file mode 100644
diff --git a/Engine/Source/Programs/LiveLinkHubEditor/Config/DefaultEngine.ini b/Engine/Source/Programs/LiveLinkHubEditor/Config/DefaultEngine.ini
deleted file mode 100644
diff --git a/Engine/Source/Programs/LiveLinkHubEditor/Config/DefaultGame.ini b/Engine/Source/Programs/LiveLinkHubEditor/Config/DefaultGame.ini
deleted file mode 100644
diff --git a/Engine/Source/Programs/LiveLinkHubEditor/LiveLinkHubEditor.uproject b/Engine/Source/Programs/LiveLinkHubEditor/LiveLinkHubEditor.uproject
deleted file mode 100644
diff --git a/Engine/Source/Programs/LiveLinkHubEditor/Source/LiveLinkHubEditor.Target.cs b/Engine/Source/Programs/LiveLinkHubEditor/Source/LiveLinkHubEditor.Target.cs
deleted file mode 100644
