Live Link Hub: Fix unclean shutdown in some code paths.

This could lead to a crash following execution of a commandlet due to the dynamic RHI not being properly torn down.

#rnx

diff --git a/Engine/Source/Programs/LiveLinkHub/Source/LiveLinkHubLauncher/Private/LiveLinkHubRun.cpp b/Engine/Source/Programs/LiveLinkHub/Source/LiveLinkHubLauncher/Private/LiveLinkHubRun.cpp
--- a/Engine/Source/Programs/LiveLinkHub/Source/LiveLinkHubLauncher/Private/LiveLinkHubRun.cpp
+++ b/Engine/Source/Programs/LiveLinkHub/Source/LiveLinkHubLauncher/Private/LiveLinkHubRun.cpp
@@ -107,6 +107,12 @@ int32 RunLiveLinkHub(const TCHAR* CommandLine)
 	int32 Result = GEngineLoop.PreInit(*FString::Printf(TEXT("%s %s"), CommandLine, TEXT("-RUN=LiveLinkHubCommandlet -Messaging -DDC=NoShared -NoShaderCompile")));
 #endif
 
+	// Ensure FEngineLoop::Exit is called for all return paths. (see EngineLoopCleanupGuard)
+	ON_SCOPE_EXIT
+	{
+		GEngineLoop.Exit();
+	};
+
 	if (Result != 0)
 	{
 		UE_LOG(LogLiveLinkHub, Error, TEXT("EngineLoop PreInit failed (%i)"), Result);
@@ -140,7 +146,5 @@ int32 RunLiveLinkHub(const TCHAR* CommandLine)
 	}
 #endif
 
-	GEngineLoop.Exit();
-
 	return Result;
 }
