Live Link Hub: Update RuntimeDependencies.ini.

#rnx

diff --git a/Engine/Source/Programs/LiveLinkHub/Source/LiveLinkHubLauncher/RuntimeDependencies.ini b/Engine/Source/Programs/LiveLinkHub/Source/LiveLinkHubLauncher/RuntimeDependencies.ini
--- a/Engine/Source/Programs/LiveLinkHub/Source/LiveLinkHubLauncher/RuntimeDependencies.ini
+++ b/Engine/Source/Programs/LiveLinkHub/Source/LiveLinkHubLauncher/RuntimeDependencies.ini
@@ -5,19 +5,28 @@
 ; Prefix - denotes exclusion
 +RuntimeDependencies="-/Content/Editor/Slate/Automation/..."
 +RuntimeDependencies="-/Content/Editor/Slate/Testing/..."
++RuntimeDependencies="/Content/BasicShapes/..."
++RuntimeDependencies="/Content/BufferVisualization/..."
 +RuntimeDependencies="/Content/EditorMaterials/..."
-+RuntimeDependencies="/Content/EditorMeshes/Camera/..."
++RuntimeDependencies="/Content/EditorMeshes/..."
 +RuntimeDependencies="/Content/EditorResources/..."
 +RuntimeDependencies="/Content/EngineDamageTypes/..."
 +RuntimeDependencies="/Content/EngineFonts/..."
 +RuntimeDependencies="/Content/EngineDebugMaterials/..."
 +RuntimeDependencies="/Content/EngineMaterials/..."
++RuntimeDependencies="/Content/EngineSounds/..."
++RuntimeDependencies="/Content/Functions/..."
 +RuntimeDependencies="/Content/Slate/..."
 
 ; Individual assets
++RuntimeDependencies="/Content/ArtTools/RenderToTexture/MaterialFunctions/CheckerPattern.uasset"
 +RuntimeDependencies="/Content/ArtTools/RenderToTexture/Meshes/S_1_Unit_Plane.uasset"
 +RuntimeDependencies="/Content/EditorLandscapeResources/DataLayer.uasset"
 +RuntimeDependencies="/Content/EditorLandscapeResources/DefaultAlphaTexture.uasset"
++RuntimeDependencies="/Content/EditorLandscapeResources/LandscapeGizmoHeight_Mat.uasset"
++RuntimeDependencies="/Content/EditorLandscapeResources/LandscapeGizmo_Mat.uasset"
++RuntimeDependencies="/Content/EditorLandscapeResources/LandscapeGizmo_Mat_Copied.uasset"
++RuntimeDependencies="/Content/EditorLandscapeResources/MF_ThicknessOutline.uasset"
 +RuntimeDependencies="/Content/EditorLandscapeResources/SplineEditorMesh.uasset"
 +RuntimeDependencies="/Content/EditorLandscapeResources/SplineEditorMeshMat.uasset"
 +RuntimeDependencies="/Content/EngineMeshes/Sphere.uasset"
@@ -27,28 +36,24 @@
 +RuntimeDependencies="/Content/EngineResources/DefaultVolumeTexture.uasset"
 +RuntimeDependencies="/Content/EngineResources/DefaultVolumeTexture2D.uasset"
 +RuntimeDependencies="/Content/EngineResources/GradientTexture0.uasset"
++RuntimeDependencies="/Content/EngineResources/LandscapeVisibilityLayerInfo.uasset"
++RuntimeDependencies="/Content/EngineResources/M_StreamingPause.uasset"
++RuntimeDependencies="/Content/EngineResources/StreamingPauseIcon.uasset"
 +RuntimeDependencies="/Content/EngineResources/WhiteSquareTexture.uasset"
 +RuntimeDependencies="/Content/EngineSky/VolumetricClouds/CloudGradientTexture.uasset"
 +RuntimeDependencies="/Content/EngineSky/VolumetricClouds/CloudWeatherTexture.uasset"
++RuntimeDependencies="/Content/EngineSky/VolumetricClouds/T_CloudMask.uasset"
++RuntimeDependencies="/Content/EngineSky/VolumetricClouds/T_CloudPattern.uasset"
 +RuntimeDependencies="/Content/EngineSky/VolumetricClouds/T_NoiseErosion.uasset"
 +RuntimeDependencies="/Content/EngineSky/VolumetricClouds/T_NoiseShape64.uasset"
++RuntimeDependencies="/Content/EngineSky/VolumetricClouds/T_Profile_08.uasset"
 +RuntimeDependencies="/Content/EngineSky/VolumetricClouds/T_VolumeNoiseErosion32.uasset"
 +RuntimeDependencies="/Content/EngineSky/VolumetricClouds/T_VolumeNoiseShape64.uasset"
++RuntimeDependencies="/Content/EngineSky/VolumetricClouds/VT_Lightning.uasset"
++RuntimeDependencies="/Content/EngineSky/VolumetricClouds/VT_PerlinWorley_Balanced.uasset"
 +RuntimeDependencies="/Content/EngineSky/VolumetricClouds/m_SimpleVolumetricCloud.uasset"
 +RuntimeDependencies="/Content/EngineSky/VolumetricClouds/m_SimpleVolumetricCloud_Inst.uasset"
-+RuntimeDependencies="/Content/Functions/Engine_MaterialFunctions01/Opacity/CameraDepthFade.uasset"
-+RuntimeDependencies="/Content/Functions/Engine_MaterialFunctions01/Shading/ConvertFromDiffSpec.uasset"
-+RuntimeDependencies="/Content/Functions/Engine_MaterialFunctions01/Shading/PowerToRoughness.uasset"
-+RuntimeDependencies="/Content/Functions/Engine_MaterialFunctions02/Math/Transform3x3Matrix.uasset"
-+RuntimeDependencies="/Content/Functions/Engine_MaterialFunctions02/PivotPainter2/Black_1x1_EXR_Texture.uasset"
-+RuntimeDependencies="/Content/Functions/Engine_MaterialFunctions02/SplitComponents.uasset"
-+RuntimeDependencies="/Content/Functions/Engine_MaterialFunctions02/Utility/BreakOutFloat2Components.uasset"
-+RuntimeDependencies="/Content/Functions/Engine_MaterialFunctions02/Utility/BreakOutFloat3Components.uasset"
-+RuntimeDependencies="/Content/Functions/Engine_MaterialFunctions02/Utility/CameraDirectionVector.uasset"
-+RuntimeDependencies="/Content/Functions/Engine_MaterialFunctions02/Utility/MakeFloat2.uasset"
-+RuntimeDependencies="/Content/Functions/Engine_MaterialFunctions02/Utility/MakeFloat3.uasset"
-+RuntimeDependencies="/Content/Functions/Engine_MaterialFunctions02/Utility/MakeFloat4.uasset"
-+RuntimeDependencies="/Content/Functions/Engine_MaterialFunctions02/WorldPositionOffset/CameraOffset.uasset"
++RuntimeDependencies="/Content/MapTemplates/Sky/DaylightAmbientCubemap.uasset"
 +RuntimeDependencies="/Content/VREditor/LaserPointer/VR_LaserPower_01.uasset"
 +RuntimeDependencies="/Content/VREditor/SnapGrid/SnapGridMaterial.uasset"
 +RuntimeDependencies="/Content/VREditor/SnapGrid/SnapGridPlaneMesh.uasset"
@@ -79,7 +84,7 @@
 +RuntimeDependencies="/Content/VREditor/TransformGizmo/Xray.uasset"
 +RuntimeDependencies="/Content/VREditor/ViewportInteractionAssetContainerData.uasset"
 
-; Plugin content (only needed for program BuildCookRun; MakeCookedEditor includes by default)
+; Plugin content
 ; FIXME: ACLPlugin still generates warnings at runtime
 +RuntimeDependencies="/Plugins/Animation/ACLPlugin/Content/..."
 +RuntimeDependencies="/Plugins/Animation/LiveLink/Content/..."
