Live Link Hub: Fix missing application icon.

#rnx

diff --git a/Engine/Source/Programs/LiveLinkHub/Resources/Mac/Assets.xcassets/AppIcon.appiconset/Contents.json b/Engine/Source/Programs/LiveLinkHub/Resources/Mac/Assets.xcassets/AppIcon.appiconset/Contents.json
deleted file mode 100644
diff --git a/Engine/Source/Programs/LiveLinkHub/Resources/Mac/Assets.xcassets/Contents.json b/Engine/Source/Programs/LiveLinkHub/Resources/Mac/Assets.xcassets/Contents.json
deleted file mode 100644
diff --git a/Engine/Source/Programs/LiveLinkHub/Resources/Windows/LiveLinkHub.rc b/Engine/Source/Programs/LiveLinkHub/Resources/Windows/LiveLinkHub.rc
deleted file mode 100644
diff --git a/Engine/Source/Programs/LiveLinkHub/Resources/Windows/make_ico.py b/Engine/Source/Programs/LiveLinkHub/Resources/Windows/make_ico.py
deleted file mode 100644
diff --git a/Engine/Source/Programs/LiveLinkHub/Source/LiveLinkHubLauncher/Private/Resources/Mac/Assets.xcassets/AppIcon.appiconset/Contents.json b/Engine/Source/Programs/LiveLinkHub/Source/LiveLinkHubLauncher/Private/Resources/Mac/Assets.xcassets/AppIcon.appiconset/Contents.json
new file mode 100644
--- /dev/null
+++ b/Engine/Source/Programs/LiveLinkHub/Source/LiveLinkHubLauncher/Private/Resources/Mac/Assets.xcassets/AppIcon.appiconset/Contents.json
@@ -0,0 +1,68 @@
+{
+  "images" : [
+    {
+      "filename" : "icon_16x16.png",
+      "idiom" : "mac",
+      "scale" : "1x",
+      "size" : "16x16"
+    },
+    {
+      "filename" : "[REDACTED_EMAIL]",
+      "idiom" : "mac",
+      "scale" : "2x",
+      "size" : "16x16"
+    },
+    {
+      "filename" : "icon_32x32.png",
+      "idiom" : "mac",
+      "scale" : "1x",
+      "size" : "32x32"
+    },
+    {
+      "filename" : "[REDACTED_EMAIL]",
+      "idiom" : "mac",
+      "scale" : "2x",
+      "size" : "32x32"
+    },
+    {
+      "filename" : "icon_128x128.png",
+      "idiom" : "mac",
+      "scale" : "1x",
+      "size" : "128x128"
+    },
+    {
+      "filename" : "[REDACTED_EMAIL]",
+      "idiom" : "mac",
+      "scale" : "2x",
+      "size" : "128x128"
+    },
+    {
+      "filename" : "icon_256x256.png",
+      "idiom" : "mac",
+      "scale" : "1x",
+      "size" : "256x256"
+    },
+    {
+      "filename" : "[REDACTED_EMAIL]",
+      "idiom" : "mac",
+      "scale" : "2x",
+      "size" : "256x256"
+    },
+    {
+      "filename" : "icon_512x512.png",
+      "idiom" : "mac",
+      "scale" : "1x",
+      "size" : "512x512"
+    },
+    {
+      "filename" : "[REDACTED_EMAIL]",
+      "idiom" : "mac",
+      "scale" : "2x",
+      "size" : "512x512"
+    }
+  ],
+  "info" : {
+    "author" : "xcode",
+    "version" : 1
+  }
+}
diff --git a/Engine/Source/Programs/LiveLinkHub/Source/LiveLinkHubLauncher/Private/Resources/Mac/Assets.xcassets/Contents.json b/Engine/Source/Programs/LiveLinkHub/Source/LiveLinkHubLauncher/Private/Resources/Mac/Assets.xcassets/Contents.json
new file mode 100644
--- /dev/null
+++ b/Engine/Source/Programs/LiveLinkHub/Source/LiveLinkHubLauncher/Private/Resources/Mac/Assets.xcassets/Contents.json
@@ -0,0 +1,6 @@
+{
+  "info" : {
+    "author" : "xcode",
+    "version" : 1
+  }
+}
diff --git a/Engine/Source/Programs/LiveLinkHub/Source/LiveLinkHubLauncher/Private/Resources/Windows/LiveLinkHub.rc b/Engine/Source/Programs/LiveLinkHub/Source/LiveLinkHubLauncher/Private/Resources/Windows/LiveLinkHub.rc
new file mode 100644
--- /dev/null
+++ b/Engine/Source/Programs/LiveLinkHub/Source/LiveLinkHubLauncher/Private/Resources/Windows/LiveLinkHub.rc
@@ -0,0 +1,88 @@
+// Copyright Epic Games, Inc. All Rights Reserved.
+
+#define APSTUDIO_READONLY_SYMBOLS
+/////////////////////////////////////////////////////////////////////////////
+//
+// Generated from the TEXTINCLUDE 2 resource.
+//
+#include <windows.h>
+#include "Runtime/Launch/Resources/Version.h"
+#include "Runtime/Launch/Resources/Windows/PCLaunch.rc"
+
+#define LIVELINK_HUB_VERSION_MAJOR 1
+#define LIVELINK_HUB_VERSION_MINOR 0
+#define LIVELINK_HUB_VERSION_PATCH 0
+
+#define LIVELINK_HUB_VERSION_STRING "v1.01"
+
+#define VER_FILEVERSION         LIVELINK_HUB_VERSION_MAJOR,LIVELINK_HUB_VERSION_MINOR,LIVELINK_HUB_VERSION_PATCH,0
+#define VER_FILEVERSION_STR     LIVELINK_HUB_VERSION_STRING "\0"
+
+#define VER_PRODUCTVERSION      ENGINE_MAJOR_VERSION,ENGINE_MINOR_VERSION,ENGINE_PATCH_VERSION,0
+#define VER_PRODUCTVERSION_STR  ENGINE_VERSION_STRING "\0"
+
+#ifdef _DEBUG
+#define VER_FILEFLAGS           0x1L // VS_FF_DEBUG
+#else
+#define VER_FILEFLAGS           0x0L
+#endif
+
+#define IDI_LIVELINK_HUB            IDICON_UEGame // 123
+
+/////////////////////////////////////////////////////////////////////////////
+#undef APSTUDIO_READONLY_SYMBOLS
+
+/////////////////////////////////////////////////////////////////////////////
+// English (U.S.) resources
+
+#if !defined(AFX_RESOURCE_DLL) || defined(AFX_TARG_ENU)
+#ifdef _WIN32
+LANGUAGE LANG_ENGLISH, SUBLANG_ENGLISH_US
+#pragma code_page(1252)
+#endif //_WIN32
+
+/////////////////////////////////////////////////////////////////////////////
+//
+// Version
+//
+
+VS_VERSION_INFO VERSIONINFO
+FILEVERSION     VER_FILEVERSION
+PRODUCTVERSION  VER_PRODUCTVERSION
+FILEFLAGSMASK   0x17L
+FILEFLAGS       VER_FILEFLAGS
+FILEOS          0x4L
+FILETYPE        0x2L
+FILESUBTYPE     0x0L
+BEGIN
+    BLOCK "StringFileInfo"
+    BEGIN
+        BLOCK "040904b0"
+        BEGIN
+            VALUE "CompanyName",      EPIC_COMPANY_NAME
+            VALUE "LegalCopyright",   EPIC_COPYRIGHT_STRING
+            VALUE "ProductName",      EPIC_PRODUCT_NAME
+            VALUE "ProductVersion",   VER_PRODUCTVERSION_STR
+            VALUE "FileDescription",  "Live Link Hub"
+            VALUE "FileVersion",      VER_FILEVERSION_STR
+            VALUE "InternalName",     "LiveLinkHub"
+            VALUE "OriginalFilename", "LiveLinkHub.exe"
+        END
+    END
+    BLOCK "VarFileInfo"
+    BEGIN
+        VALUE "Translation", 0x409, 1200
+    END
+END
+
+/////////////////////////////////////////////////////////////////////////////
+//
+// Icon
+//
+
+// Icon with lowest ID value placed first to ensure application icon
+// remains consistent on all systems.
+IDI_LIVELINK_HUB    ICON    "LiveLinkHub.ico"
+
+#endif    // English (U.S.) resources
+/////////////////////////////////////////////////////////////////////////////
diff --git a/Engine/Source/Programs/LiveLinkHub/Source/LiveLinkHubLauncher/Private/Resources/Windows/make_ico.py b/Engine/Source/Programs/LiveLinkHub/Source/LiveLinkHubLauncher/Private/Resources/Windows/make_ico.py
new file mode 100644
--- /dev/null
+++ b/Engine/Source/Programs/LiveLinkHub/Source/LiveLinkHubLauncher/Private/Resources/Windows/make_ico.py
@@ -0,0 +1,15 @@
+# Copyright Epic Games, Inc. All Rights Reserved.
+
+from PIL import Image
+
+def MakeIcon(folder, output):
+    sizes = [256, 48, 32, 24, 16]
+    images = []
+    for size in sizes:
+        print(f"Loading {folder}/icon_{size}x{size}.png")
+        image = Image.open(f"{folder}/icon_{size}x{size}.png")
+        images.append(image)
+    print(f"Creating {output}")
+    images[0].save(f"{output}", append_images=images)
+
+MakeIcon("LiveLinkHub", "LiveLinkHub.ico")
