Live Link Hub: Custom hook for plugin browser restart behavior.

#rnx

diff --git a/Engine/Source/Programs/LiveLinkHub/Source/LiveLinkHubLauncher/Private/LiveLinkHubRun.cpp b/Engine/Source/Programs/LiveLinkHub/Source/LiveLinkHubLauncher/Private/LiveLinkHubRun.cpp
--- a/Engine/Source/Programs/LiveLinkHub/Source/LiveLinkHubLauncher/Private/LiveLinkHubRun.cpp
+++ b/Engine/Source/Programs/LiveLinkHub/Source/LiveLinkHubLauncher/Private/LiveLinkHubRun.cpp
@@ -13,12 +13,14 @@
 #include "HAL/PlatformMisc.h"
 #include "HAL/PlatformSplash.h"
 #include "ILiveLinkHubModule.h"
+#include "IPluginBrowser.h"
 #include "LaunchEngineLoop.h"
 #include "LiveLinkHubPluginHelpers.h"
 #include "Misc/CommandLine.h"
 #include "Misc/CoreDelegates.h"
 #include "Misc/Paths.h"
 #include "Modules/ModuleManager.h"
+#include "UnrealEdMisc.h"
 
 #ifndef WITH_ASSET_LOADING_AUDIT
 #define WITH_ASSET_LOADING_AUDIT 0
@@ -57,6 +59,7 @@ int32 RunLiveLinkHub(const TCHAR* CommandLine)
 #endif
 
 	bool bLauncherDistribution = false;
+	bool bPendingRestart = false;
 
 #if !IS_PROGRAM
 	const TCHAR* const DevelopmentProjectPath = TEXT("../../Source/Programs/LiveLinkHub/LiveLinkHub.uproject");
@@ -92,7 +95,17 @@ int32 RunLiveLinkHub(const TCHAR* CommandLine)
 
 	// Plugin directory config save/load
 	FCoreDelegates::TSConfigReadyForUse().AddStatic(LiveLinkHub::PluginHelpers::RestoreSavedPluginDirectories);
-	FCoreDelegates::OnAllModuleLoadingPhasesComplete.AddStatic(LiveLinkHub::PluginHelpers::RegisterPluginDirectoriesChangedHandler);
+	FCoreDelegates::OnAllModuleLoadingPhasesComplete.AddLambda([&bPendingRestart]()
+	{
+		LiveLinkHub::PluginHelpers::RegisterPluginDirectoriesChangedHandler();
+		
+		IPluginBrowser& PluginBrowser = IPluginBrowser::Get();
+		PluginBrowser.OnRestartClicked().BindLambda([&bPendingRestart]()
+		{
+			bPendingRestart = true;
+			RequestEngineExit(TEXT("Plugin browser requested restart"));
+		});
+	});
 #endif
 
 #if IS_PROGRAM
@@ -111,6 +124,23 @@ int32 RunLiveLinkHub(const TCHAR* CommandLine)
 	ON_SCOPE_EXIT
 	{
 		GEngineLoop.Exit();
+
+		// Handle this after FEngineLoop::Cleanup, which can take several seconds.
+		// We want to defer this as long as possible to avoid/minimize concurrent execution.
+		if (bPendingRestart)
+		{
+			const FString ExeFilename = FUnrealEdMisc::GetProjectEditorBinaryPath();
+			FProcHandle Handle = FPlatformProcess::CreateProc(*ExeFilename, FCommandLine::Get(), true, false, false, NULL, 0, NULL, NULL);
+			const bool bSuccess = Handle.IsValid();
+			if (bSuccess)
+			{
+				FPlatformProcess::CloseProc(Handle);
+			}
+			else
+			{
+				UE_LOG(LogLiveLinkHub, Error, TEXT("Pending restart: FPlatformProcess::CreateProc failed"));
+			}
+		}
 	};
 
 	if (Result != 0)
