Fix CIS warnings when generating the LiveLinkHub DDC build
Add LiveLinkHub to the InstalledEngineBuild list of projects to build the DDC for

diff --git a/Engine/Build/InstalledEngineFilters.xml b/Engine/Build/InstalledEngineFilters.xml
--- a/Engine/Build/InstalledEngineFilters.xml
+++ b/Engine/Build/InstalledEngineFilters.xml
@@ -67,6 +67,7 @@
 		Templates/TP_DMXBP/TP_DMXBP.uproject
 		Templates/TP_MotionDesignBP/TP_MotionDesignBP.uproject
 		Samples/StarterContent/StarterContent.uproject
+		Engine/Source/Programs/LiveLinkHub/LiveLinkHub.uproject
 	</Property>
 	<!-- Win64 specific DDC -->
 	<Property Name="ProjectsToBuildDDCWin64" Value="$(ProjectsToBuildDDC)">
diff --git a/Engine/Source/Programs/LiveLinkHub/Config/DefaultEngine.ini b/Engine/Source/Programs/LiveLinkHub/Config/DefaultEngine.ini
--- a/Engine/Source/Programs/LiveLinkHub/Config/DefaultEngine.ini
+++ b/Engine/Source/Programs/LiveLinkHub/Config/DefaultEngine.ini
@@ -40,3 +40,12 @@ DefaultEditorTarget=LiveLinkHub
 ; Not needed at the moment, can be removed if we enable the TextureFormatOodle plugin
 [AlternateTextureCompression]
 TextureCompressionFormat=""
+
+[/Script/InterchangeEngine.InterchangeProjectSettings]
+ContentImportSettings=()
+SceneImportSettings=()
+FilePickerClass=None
+bStaticMeshUseSmoothEdgesIfSmoothingInformationIsMissing=True
+GenericPipelineClass=None
+ConverterDefaultPipeline=None
+
