Switch to engine animation codecs

diff --git a/Engine/Source/Programs/LiveLinkHub/Config/DefaultEngine.ini b/Engine/Source/Programs/LiveLinkHub/Config/DefaultEngine.ini
--- a/Engine/Source/Programs/LiveLinkHub/Config/DefaultEngine.ini
+++ b/Engine/Source/Programs/LiveLinkHub/Config/DefaultEngine.ini
@@ -49,3 +49,7 @@ bStaticMeshUseSmoothEdgesIfSmoothingInformationIsMissing=True
 GenericPipelineClass=None
 ConverterDefaultPipeline=None
 
+; Use the engine fallback codecs since engine plugins are disabled
+[Animation.DefaultObjectSettings]
+BoneCompressionSettings="/Engine/Animation/DefaultAnimBoneCompressionSettings"
+CurveCompressionSettings="/Engine/Animation/DefaultAnimCurveCompressionSettings"
