Live Link Hub: Telemetry cleanup.

Ensure project name is correctly captured by crash context, set AppIdOverride, set ProjectID.

Corresponding standalone build automation changes.

#rnx

diff --git a/Engine/Source/Programs/LiveLinkHub/Config/DefaultEditor.ini b/Engine/Source/Programs/LiveLinkHub/Config/DefaultEditor.ini
--- a/Engine/Source/Programs/LiveLinkHub/Config/DefaultEditor.ini
+++ b/Engine/Source/Programs/LiveLinkHub/Config/DefaultEditor.ini
@@ -12,3 +12,7 @@ bCanModifyProjectPluginDirectoriesFromBrowser=false
 ; Hide plugins that aren't "Installed" or "External" to cut down on noise.
 +HidePluginCategoriesFromBrowser=All
 +HidePluginCategoriesFromBrowser=Built-In
+
+
+[Analytics]
+AppIdOverride=LiveLinkHub
diff --git a/Engine/Source/Programs/LiveLinkHub/Config/DefaultGame.ini b/Engine/Source/Programs/LiveLinkHub/Config/DefaultGame.ini
--- a/Engine/Source/Programs/LiveLinkHub/Config/DefaultGame.ini
+++ b/Engine/Source/Programs/LiveLinkHub/Config/DefaultGame.ini
@@ -8,7 +8,13 @@ bStageUAT=false
 bStagePython=false
 bIsForExternalDistribution=true
 
+
 [/Script/LiveLink.LiveLinkSettings]
 DefaultSettingsClass=/Script/LiveLinkHub.LiveLinkHubSubjectSettings
 bPreProcessRebroadcastFrames=true
 bTranslateRebroadcastFrames=true
+
+
+[/Script/EngineSettings.GeneralProjectSettings]
+ProjectID=914C155297C9471A9E34100517697A75
+ProjectName=LiveLinkHub
diff --git a/Engine/Source/Programs/LiveLinkHub/Source/LiveLinkHubLauncher/Private/LiveLinkHubRun.cpp b/Engine/Source/Programs/LiveLinkHub/Source/LiveLinkHubLauncher/Private/LiveLinkHubRun.cpp
--- a/Engine/Source/Programs/LiveLinkHub/Source/LiveLinkHubLauncher/Private/LiveLinkHubRun.cpp
+++ b/Engine/Source/Programs/LiveLinkHub/Source/LiveLinkHubLauncher/Private/LiveLinkHubRun.cpp
@@ -85,14 +85,6 @@ int32 RunLiveLinkHub(const TCHAR* CommandLine)
 	// Used by Live Coding, among other things.
 	FPlatformMisc::SetUBTTargetName(TEXT("LiveLinkHub"));
 
-	// We need to specify this manually, because it isn't inferred from the command line.
-	// But we can't set it too early, because LaunchSetGameName clears it.
-	FCoreDelegates::OnInit.AddStatic([]()
-	{
-		// Required for OpenXR instance registration.
-		FApp::SetProjectName(TEXT("LiveLinkHub"));
-	});
-
 	// Plugin directory config save/load
 	FCoreDelegates::TSConfigReadyForUse().AddStatic(LiveLinkHub::PluginHelpers::RestoreSavedPluginDirectories);
 	FCoreDelegates::OnAllModuleLoadingPhasesComplete.AddLambda([&bPendingRestart]()
@@ -115,7 +107,9 @@ int32 RunLiveLinkHub(const TCHAR* CommandLine)
 
 	// Start up the main loop, adding some extra command line arguments:
 #if !IS_PROGRAM
-	int32 Result = GEngineLoop.PreInit(*FString::Printf(TEXT("%s %s"), CommandLine, TEXT("-xrtrackingonly")));
+	const TCHAR* const ExtraArgs = TEXT("-xrtrackingonly");
+	int32 Result = GEngineLoop.PreInit(*FString::Printf(TEXT("%s %s %s"),
+		*FPaths::GetProjectFilePath(), CommandLine, ExtraArgs));
 #else
 	int32 Result = GEngineLoop.PreInit(*FString::Printf(TEXT("%s %s"), CommandLine, TEXT("-RUN=LiveLinkHubCommandlet -Messaging -DDC=NoShared -NoShaderCompile")));
 #endif
