Disable joystick for slate and editor viewport for Live Link Hub

diff --git a/Engine/Source/Programs/LiveLinkHub/Config/DefaultEditor.ini b/Engine/Source/Programs/LiveLinkHub/Config/DefaultEditor.ini
--- a/Engine/Source/Programs/LiveLinkHub/Config/DefaultEditor.ini
+++ b/Engine/Source/Programs/LiveLinkHub/Config/DefaultEditor.ini
@@ -16,3 +16,6 @@ bCanModifyProjectPluginDirectoriesFromBrowser=false
 
 [Analytics]
 AppIdOverride=LiveLinkHub
+
+[/Script/UnrealEd.LevelEditorViewportSettings]
+bLevelEditorJoystickControls=False
diff --git a/Engine/Source/Programs/LiveLinkHub/Config/DefaultEngine.ini b/Engine/Source/Programs/LiveLinkHub/Config/DefaultEngine.ini
--- a/Engine/Source/Programs/LiveLinkHub/Config/DefaultEngine.ini
+++ b/Engine/Source/Programs/LiveLinkHub/Config/DefaultEngine.ini
@@ -33,6 +33,7 @@ bShareMaterialShaderCode=false
 
 [ConsoleVariables]
 DetailsPanel.Style.EnableCardLayout=true
+Slate.EnableGamepadEditorNavigation=false
 
 [/Script/BuildSettings.BuildSettings]
 DefaultEditorTarget=LiveLinkHub
