Clean up LiveLinkHubRun - Move LiveLinkHub's loading phase to PreDefault so it can register itself as a LiveLinkClient modular feature before other modules.

Bypassing SubmitTool because it can't run when modifying LiveLinkHubRun

diff --git a/Engine/Plugins/Animation/LiveLinkHub/LiveLinkHub.uplugin b/Engine/Plugins/Animation/LiveLinkHub/LiveLinkHub.uplugin
--- a/Engine/Plugins/Animation/LiveLinkHub/LiveLinkHub.uplugin
+++ b/Engine/Plugins/Animation/LiveLinkHub/LiveLinkHub.uplugin
@@ -27,7 +27,7 @@
 			"Name": "LiveLinkHub",
 			"Type": "EditorAndProgram",
 			"ProgramAllowList": [ "LiveLinkHub" ],
-			"LoadingPhase": "Default"
+			"LoadingPhase": "PreDefault"
 		},
 		{
 			"Name": "LiveLinkHubEditor",
diff --git a/Engine/Plugins/Animation/LiveLinkHub/Source/LiveLinkHub/Private/Clients/LiveLinkHubProvider.cpp b/Engine/Plugins/Animation/LiveLinkHub/Source/LiveLinkHub/Private/Clients/LiveLinkHubProvider.cpp
--- a/Engine/Plugins/Animation/LiveLinkHub/Source/LiveLinkHub/Private/Clients/LiveLinkHubProvider.cpp
+++ b/Engine/Plugins/Animation/LiveLinkHub/Source/LiveLinkHub/Private/Clients/LiveLinkHubProvider.cpp
@@ -88,8 +88,11 @@ FLiveLinkHubProvider::FLiveLinkHubProvider(const TSharedRef<ILiveLinkHubSessionM
 
 	CreateMessageEndpoint(EndpointBuilder);
 
-	const double ValidateConnectionsRate = GetDefault<ULiveLinkSettings>()->MessageBusPingRequestFrequency;
-	GEditor->GetTimerManager()->SetTimer(ValidateConnectionsTimer, FTimerDelegate::CreateRaw(this, &FLiveLinkHubProvider::ValidateConnections), ValidateConnectionsRate, true);
+	FCoreDelegates::OnPostEngineInit.AddLambda([this]()
+	{
+		const double ValidateConnectionsRate = GetDefault<ULiveLinkSettings>()->MessageBusPingRequestFrequency;
+        GEditor->GetTimerManager()->SetTimer(ValidateConnectionsTimer, FTimerDelegate::CreateRaw(this, &FLiveLinkHubProvider::ValidateConnections), ValidateConnectionsRate, true);
+	});
 }
 
 FLiveLinkHubProvider::~FLiveLinkHubProvider()
diff --git a/Engine/Plugins/Animation/LiveLinkHub/Source/LiveLinkHub/Private/LiveLinkHub.cpp b/Engine/Plugins/Animation/LiveLinkHub/Source/LiveLinkHub/Private/LiveLinkHub.cpp
--- a/Engine/Plugins/Animation/LiveLinkHub/Source/LiveLinkHub/Private/LiveLinkHub.cpp
+++ b/Engine/Plugins/Animation/LiveLinkHub/Source/LiveLinkHub/Private/LiveLinkHub.cpp
@@ -29,9 +29,19 @@
 
 #define LOCTEXT_NAMESPACE "LiveLinkHub"
 
+namespace LiveLinkHubInternalUtils
+{
+	/** Is this a standalone distributed build on the EGS? */
+	bool IsDistributedBuild()
+	{
+		return FPaths::FileExists(TEXT("../../../LiveLinkHub/LiveLinkHub.uproject"));
+	}
+}
 
-void FLiveLinkHub::Preinitialize(FLiveLinkHubTicker& Ticker)
+void FLiveLinkHub::Initialize(FLiveLinkHubTicker& Ticker)
 {
+	TRACE_CPUPROFILER_EVENT_SCOPE(FLiveLinkHub::Initialize);
+
 	// We must register the livelink client first since we might rely on the modular feature to initialize the controllers/managers.
 	if (GetDefault<ULiveLinkHubSettings>()->bTickOnGameThread)
 	{
@@ -41,18 +51,8 @@ void FLiveLinkHub::Preinitialize(FLiveLinkHubTicker& Ticker)
 	{
 		LiveLinkHubClient = MakeShared<FLiveLinkHubClient>(AsShared(), Ticker.OnTick());
 	}
-	
-	IModularFeatures::Get().RegisterModularFeature(ILiveLinkClient::ModularFeatureName, LiveLinkHubClient.Get());
-}
 
-void FLiveLinkHub::Initialize(bool bLauncherDistribution)
-{
-	TRACE_CPUPROFILER_EVENT_SCOPE(FLiveLinkHub::Initialize);
-
-#if IS_PROGRAM
-	// Re-enable this since we've disabled it to avoid the creation of the console window.
-	GIsSilent = false;
-#endif
+	IModularFeatures::Get().RegisterModularFeature(ILiveLinkClient::ModularFeatureName, LiveLinkHubClient.Get());
 
 	SessionManager = MakeShared<FLiveLinkHubSessionManager>();
 	LiveLinkProvider = MakeShared<FLiveLinkHubProvider>(SessionManager.ToSharedRef());
@@ -60,7 +60,7 @@ void FLiveLinkHub::Initialize(bool bLauncherDistribution)
 	FModuleManager::Get().LoadModule("Settings");
 	FModuleManager::Get().LoadModule("StatusBar");
 
-	if (bLauncherDistribution)
+	if (LiveLinkHubInternalUtils::IsDistributedBuild())
 	{
 		IAssetRegistry& AssetRegistry = FModuleManager::LoadModuleChecked<FAssetRegistryModule>(AssetRegistryConstants::ModuleName).Get();
 		const FString FilePath = FPaths::Combine(FPlatformProcess::UserSettingsDir(), *FApp::GetEpicProductIdentifier(), TEXT("LiveLinkHub"), TEXT("Content"));
@@ -80,10 +80,14 @@ void FLiveLinkHub::Initialize(bool bLauncherDistribution)
 		SubjectController = MakeShared<FLiveLinkHubSubjectController>();
 	}
 
+	FCoreDelegates::OnPostEngineInit.AddRaw(this, &FLiveLinkHub::InitializePostEngineInit);
+}
 
+void FLiveLinkHub::InitializePostEngineInit()
+{
 	FLiveLinkHubCommands::Register();
 	BindCommands();
-	
+
 	FString LiveLinkHubLayoutIni = GConfig->GetConfigFilename(TEXT("LiveLinkHubLayout"));
 	WindowController = MakeShared<FLiveLinkHubWindowController>(FLiveLinkHubWindowInitParams{ LiveLinkHubLayoutIni });
 	WindowController->RestoreLayout();
diff --git a/Engine/Plugins/Animation/LiveLinkHub/Source/LiveLinkHub/Private/LiveLinkHub.h b/Engine/Plugins/Animation/LiveLinkHub/Source/LiveLinkHub/Private/LiveLinkHub.h
--- a/Engine/Plugins/Animation/LiveLinkHub/Source/LiveLinkHub/Private/LiveLinkHub.h
+++ b/Engine/Plugins/Animation/LiveLinkHub/Source/LiveLinkHub/Private/LiveLinkHub.h
@@ -49,13 +49,9 @@ public:
 	//~ End ILiveLinkHub interface
 
 public:
-	/** First phase initialization; must precede FEngineLoop::Init(). */
-	void Preinitialize(class FLiveLinkHubTicker& Ticker);
 	/** 
-	 * Launch the slate application and initialize its components. 
-	 * @param bLauncherDistribution Whether the hub should be initialized as a launcher distributed LiveLinkHub.
-	 */
-	void Initialize(bool bLauncherDistribution);
+	 * Initialize LiveLinkHub. */
+	void Initialize(class FLiveLinkHubTicker& Ticker);
 	/** Tick the hub. */
 	void Tick();
 	/** Get the root window that hosts the hub's slate application. */
@@ -76,6 +72,9 @@ public:
 	TSharedPtr<ILiveLinkHubSessionManager> GetSessionManager() const;
 	
 private:
+	/** Complete the initialization after core engine systems have been initialized. */
+	void InitializePostEngineInit();
+
 	//~ LiveLink Client delegates
 	void OnStaticDataReceived_AnyThread(const FLiveLinkSubjectKey& InSubjectKey, TSubclassOf<ULiveLinkRole> InRole, const FLiveLinkStaticDataStruct& InStaticDataStruct) const;
 	void OnFrameDataReceived_AnyThread(const FLiveLinkSubjectKey& InSubjectKey, const FLiveLinkFrameDataStruct& InFrameDataStruct) const;
diff --git a/Engine/Plugins/Animation/LiveLinkHub/Source/LiveLinkHub/Private/LiveLinkHubApplication.h b/Engine/Plugins/Animation/LiveLinkHub/Source/LiveLinkHub/Private/LiveLinkHubApplication.h
deleted file mode 100644
diff --git a/Engine/Plugins/Animation/LiveLinkHub/Source/LiveLinkHub/Private/LiveLinkHubModule.cpp b/Engine/Plugins/Animation/LiveLinkHub/Source/LiveLinkHub/Private/LiveLinkHubModule.cpp
--- a/Engine/Plugins/Animation/LiveLinkHub/Source/LiveLinkHub/Private/LiveLinkHubModule.cpp
+++ b/Engine/Plugins/Animation/LiveLinkHub/Source/LiveLinkHub/Private/LiveLinkHubModule.cpp
@@ -3,9 +3,6 @@
 #include "LiveLinkHubModule.h"
 
 #include "Clients/LiveLinkHubProvider.h"
-#include "Editor/EditorPerformanceSettings.h"
-#include "LiveLinkHubApplication.h"
-#include "LiveLinkHubLog.h"
 #include "LiveLinkHubSubjectSettings.h"
 #include "LiveLinkSettings.h"
 #include "Misc/ConfigCacheIni.h"
@@ -17,12 +14,14 @@
 #include "Settings/LiveLinkHubSettingsCustomization.h"
 #include "Settings/LiveLinkSettingsCustomization.h"
 #include "Subjects/LiveLinkHubSubjectSettingsDetailsCustomization.h"
+#include "Textures/SlateIcon.h"
+#include "Widgets/Text/STextBlock.h"
 
 #define LOCTEXT_NAMESPACE "LiveLinkHubModule"
 
-void FLiveLinkHubModule::PreinitializeLiveLinkHub()
+void FLiveLinkHubModule::StartupModule()
 {
-	check(!LiveLinkHub);
+	TRACE_CPUPROFILER_EVENT_SCOPE(StartLiveLinkHub);
 
 	if (!GetDefault<ULiveLinkHubSettings>()->bTickOnGameThread)
 	{
@@ -30,34 +29,8 @@ void FLiveLinkHubModule::PreinitializeLiveLinkHub()
 	}
 
 	LiveLinkHub = MakeShared<FLiveLinkHub>();
-	LiveLinkHub->Preinitialize(Ticker);
-}
-
-void FLiveLinkHubModule::StartLiveLinkHub(bool bLauncherDistribution)
-{
-	TRACE_CPUPROFILER_EVENT_SCOPE(StartLiveLinkHub);
-	checkf(LiveLinkHub, TEXT("Ensure PreinitializeLiveLinkHub was called first"));
-
-	LiveLinkHub->Initialize(bLauncherDistribution);
-
-#if IS_PROGRAM
-	LiveLinkHubLoop(LiveLinkHub);
-#endif
-}
-
-void FLiveLinkHubModule::ShutdownLiveLinkHub()
-{
-	if (!GetDefault<ULiveLinkHubSettings>()->bTickOnGameThread)
-	{
-		Ticker.Exit();
-		Ticker.Stop();
-	}
-
-	LiveLinkHub.Reset();
-}
+	LiveLinkHub->Initialize(Ticker);
 
-void FLiveLinkHubModule::StartupModule()
-{
 	FPropertyEditorModule& PropertyModule = FModuleManager::LoadModuleChecked<FPropertyEditorModule>("PropertyEditor");
 
 	PropertyModule.RegisterCustomClassLayout(ULiveLinkHubSettings::StaticClass()->GetFName(), FOnGetDetailCustomizationInstance::CreateStatic(&FLiveLinkHubSettingsCustomization::MakeInstance));
@@ -93,6 +66,14 @@ void FLiveLinkHubModule::ShutdownModule()
 			PropertyEditorModule->UnregisterCustomClassLayout(ULiveLinkHubSubjectSettings::StaticClass()->GetFName());
 		}
 	}
+
+	if (!GetDefault<ULiveLinkHubSettings>()->bTickOnGameThread)
+	{
+		Ticker.Exit();
+		Ticker.Stop();
+	}
+
+	LiveLinkHub.Reset();
 }
 
 TSharedPtr<FLiveLinkHub> FLiveLinkHubModule::GetLiveLinkHub() const
@@ -131,4 +112,5 @@ TSharedPtr<ILiveLinkHubSessionManager> FLiveLinkHubModule::GetSessionManager() c
 }
 
 IMPLEMENT_MODULE(FLiveLinkHubModule, LiveLinkHub);
+
 #undef LOCTEXT_NAMESPACE /* LiveLinkHubModule */
diff --git a/Engine/Plugins/Animation/LiveLinkHub/Source/LiveLinkHub/Private/LiveLinkHubModule.h b/Engine/Plugins/Animation/LiveLinkHub/Source/LiveLinkHub/Private/LiveLinkHubModule.h
--- a/Engine/Plugins/Animation/LiveLinkHub/Source/LiveLinkHub/Private/LiveLinkHubModule.h
+++ b/Engine/Plugins/Animation/LiveLinkHub/Source/LiveLinkHub/Private/LiveLinkHubModule.h
@@ -19,9 +19,6 @@ class FLiveLinkHubModule : public ILiveLinkHubModule
 {
 public:
 	//~ Begin ILiveLinkHubModule interface
-	virtual void PreinitializeLiveLinkHub() override;
-	virtual void StartLiveLinkHub(bool bLauncherDistribution = false) override;
-	virtual void ShutdownLiveLinkHub() override;
 	virtual void StartupModule() override;
 	virtual void ShutdownModule() override;
 	//~ End ILiveLinkHubModule interface
diff --git a/Engine/Plugins/Animation/LiveLinkHub/Source/LiveLinkHub/Public/ILiveLinkHubModule.h b/Engine/Plugins/Animation/LiveLinkHub/Source/LiveLinkHub/Public/ILiveLinkHubModule.h
--- a/Engine/Plugins/Animation/LiveLinkHub/Source/LiveLinkHub/Public/ILiveLinkHubModule.h
+++ b/Engine/Plugins/Animation/LiveLinkHub/Source/LiveLinkHub/Public/ILiveLinkHubModule.h
@@ -6,12 +6,13 @@
 
 class LIVELINKHUB_API ILiveLinkHubModule : public IModuleInterface
 {
-public:
-	virtual void PreinitializeLiveLinkHub() = 0;
+	UE_DEPRECATED(5.6, "There's no need to call this method anymore, it's now handled by StartupModule and ShutdownModule.")
+	virtual void PreinitializeLiveLinkHub() {}
 
+	UE_DEPRECATED(5.6, "There's no need to call this method anymore, it's now handled by StartupModule and ShutdownModule.")
 	/** Launch the slate application hosting the live link hub. */
-	virtual void StartLiveLinkHub(bool bLauncherDistribution = false) = 0;
+	virtual void StartLiveLinkHub(bool bLauncherDistribution = false) {}
 
-	virtual void ShutdownLiveLinkHub() = 0;
+	UE_DEPRECATED(5.6, "There's no need to call this method anymore, it's now handled by StartupModule and ShutdownModule.")
+	virtual void ShutdownLiveLinkHub() {}
 };
-
diff --git a/Engine/Source/Programs/LiveLinkHub/Source/LiveLinkHubLauncher/Private/LiveLinkHubCommandlet.h b/Engine/Source/Programs/LiveLinkHub/Source/LiveLinkHubLauncher/Private/LiveLinkHubCommandlet.h
deleted file mode 100644
diff --git a/Engine/Source/Programs/LiveLinkHub/Source/LiveLinkHubLauncher/Private/LiveLinkHubRun.cpp b/Engine/Source/Programs/LiveLinkHub/Source/LiveLinkHubLauncher/Private/LiveLinkHubRun.cpp
--- a/Engine/Source/Programs/LiveLinkHub/Source/LiveLinkHubLauncher/Private/LiveLinkHubRun.cpp
+++ b/Engine/Source/Programs/LiveLinkHub/Source/LiveLinkHubLauncher/Private/LiveLinkHubRun.cpp
@@ -5,20 +5,16 @@
 #include "Editor.h"
 #include "Editor/EditorEngine.h"
 #include "Misc/App.h"
+#include "Misc/CommandLine.h"
+#include "Misc/CoreDelegates.h"
 #include "Misc/PackageName.h"
 #include "Misc/Paths.h"
 #include "HAL/FileManager.h"
-#include "HAL/PlatformApplicationMisc.h"
-#include "HAL/PlatformFileManager.h"
 #include "HAL/PlatformMisc.h"
 #include "HAL/PlatformSplash.h"
-#include "ILiveLinkHubModule.h"
 #include "IPluginBrowser.h"
 #include "LaunchEngineLoop.h"
 #include "LiveLinkHubPluginHelpers.h"
-#include "Misc/CommandLine.h"
-#include "Misc/CoreDelegates.h"
-#include "Misc/Paths.h"
 #include "Modules/ModuleManager.h"
 #include "UnrealEdMisc.h"
 
@@ -58,28 +54,24 @@ int32 RunLiveLinkHub(const TCHAR* CommandLine)
 		});
 #endif
 
-	bool bLauncherDistribution = false;
 	bool bPendingRestart = false;
 
-#if !IS_PROGRAM
 	const TCHAR* const DevelopmentProjectPath = TEXT("../../Source/Programs/LiveLinkHub/LiveLinkHub.uproject");
 	const TCHAR* const StagedProjectPath = TEXT("../../../LiveLinkHub/LiveLinkHub.uproject");
 
 	if (FPaths::FileExists(DevelopmentProjectPath))
 	{
+		// This is a 'in-tree' version of LiveLinkHub, meaning it could be a development build or the version of LLH that is shipped with UnrealEngine.
 		FPaths::SetProjectFilePath(DevelopmentProjectPath);
 	}
 	else if (FPaths::FileExists(StagedProjectPath))
 	{
-		bLauncherDistribution = true;
+		// This is build of LiveLinkHub that's distributed on the EGS.
 		FPaths::SetProjectFilePath(StagedProjectPath);
-	}
 
-	if (bLauncherDistribution)
-	{
 		// This is needed to ensure that content generated by the hub will be discoverable (ie. LiveLink recordings)
-		const FString GameContentPath = FPaths::Combine(FPlatformProcess::UserSettingsDir(), *FApp::GetEpicProductIdentifier(), TEXT("LiveLinkHub"), TEXT("Content"));
-		FPackageName::RegisterMountPoint(TEXT("/Game/"), GameContentPath);
+        const FString GameContentPath = FPaths::Combine(FPlatformProcess::UserSettingsDir(), *FApp::GetEpicProductIdentifier(), TEXT("LiveLinkHub"), TEXT("Content"));
+        FPackageName::RegisterMountPoint(TEXT("/Game/"), GameContentPath);
 	}
 
 	// Used by Live Coding, among other things.
@@ -98,21 +90,12 @@ int32 RunLiveLinkHub(const TCHAR* CommandLine)
 			RequestEngineExit(TEXT("Plugin browser requested restart"));
 		});
 	});
-#endif
-
-#if IS_PROGRAM
-	// Disable this when going through PreInit to prevent the console window from appearing.
-	GIsSilent = true;
-#endif
 
 	// Start up the main loop, adding some extra command line arguments:
-#if !IS_PROGRAM
 	const TCHAR* const ExtraArgs = TEXT("-xrtrackingonly");
 	int32 Result = GEngineLoop.PreInit(*FString::Printf(TEXT("%s %s %s"),
 		*FPaths::GetProjectFilePath(), CommandLine, ExtraArgs));
-#else
-	int32 Result = GEngineLoop.PreInit(*FString::Printf(TEXT("%s %s"), CommandLine, TEXT("-RUN=LiveLinkHubCommandlet -Messaging -DDC=NoShared -NoShaderCompile")));
-#endif
+
 
 	// Ensure FEngineLoop::Exit is called for all return paths. (see EngineLoopCleanupGuard)
 	ON_SCOPE_EXIT
@@ -147,17 +130,11 @@ int32 RunLiveLinkHub(const TCHAR* CommandLine)
 		return Result;
 	}
 
-#if !IS_PROGRAM
 	{
-		ILiveLinkHubModule& HubModule = FModuleManager::Get().LoadModuleChecked<ILiveLinkHubModule>("LiveLinkHub");
-		HubModule.PreinitializeLiveLinkHub();
-
 		Result = GEngineLoop.Init();
 
 		GEditor->ShouldDisableCPUThrottlingDelegates.Add(UEditorEngine::FShouldDisableCPUThrottling::CreateLambda([]() { return true; }));
 
-		HubModule.StartLiveLinkHub(bLauncherDistribution);
-
 		// Hide the splash screen now that everything is ready to go
 		FPlatformSplash::Hide();
 
@@ -165,10 +142,7 @@ int32 RunLiveLinkHub(const TCHAR* CommandLine)
 		{
 			GEngineLoop.Tick();
 		}
-
-		HubModule.ShutdownLiveLinkHub();
 	}
-#endif
 
 	return Result;
 }
