Live Link Hub: Permit toggling the enabled state of individual plugins via the Plugin Browser.

diff --git a/Engine/Source/Programs/LiveLinkHub/Config/DefaultEditor.ini b/Engine/Source/Programs/LiveLinkHub/Config/DefaultEditor.ini
--- a/Engine/Source/Programs/LiveLinkHub/Config/DefaultEditor.ini
+++ b/Engine/Source/Programs/LiveLinkHub/Config/DefaultEditor.ini
@@ -1,8 +1,6 @@
 [EditorSettings]
 ; We don't allow creating new plugins.
 bCanCreatePluginsFromBrowser=false
-; We don't allow toggling individual plugins.
-bCanModifyPluginsFromBrowser=false
 
 ; We allow modifying user plugin paths...
 bCanModifyUserPluginDirectoriesFromBrowser=true
