Capture Manager plugin should not be optional

The plugin shouldn't be optional as it is already hidden via a flag

diff --git a/Engine/Source/Programs/LiveLinkHub/Source/LiveLinkHub.Target.cs b/Engine/Source/Programs/LiveLinkHub/Source/LiveLinkHub.Target.cs
--- a/Engine/Source/Programs/LiveLinkHub/Source/LiveLinkHub.Target.cs
+++ b/Engine/Source/Programs/LiveLinkHub/Source/LiveLinkHub.Target.cs
@@ -62,7 +62,7 @@ public class LiveLinkHubTarget : TargetRules
 
 		if (bEnableCaptureManagerPlugin)
 		{
-			OptionalPlugins.AddRange(new string[]
+			EnablePlugins.AddRange(new string[]
 			{
 				"CaptureManagerApp",
 				"CaptureManagerCore"
