Enable Capture Manager plugins by default

diff --git a/Engine/Source/Programs/LiveLinkHub/Source/LiveLinkHub.Target.cs b/Engine/Source/Programs/LiveLinkHub/Source/LiveLinkHub.Target.cs
--- a/Engine/Source/Programs/LiveLinkHub/Source/LiveLinkHub.Target.cs
+++ b/Engine/Source/Programs/LiveLinkHub/Source/LiveLinkHub.Target.cs
@@ -14,10 +14,6 @@ public class LiveLinkHubTarget : TargetRules
 	[CommandLine("-SingleModulePlatform")]
 	public bool bSingleModulePlatform = false;
 
-	// Whether to enable building Capture Manager plugin.
-	[CommandLine("-EnableCaptureManagerPlugin=")]
-	public bool bEnableCaptureManagerPlugin = true;
-
 	// Whether to enable building MetaHuman LiveLink plugin.
 	[CommandLine("-EnableMetaHumanLiveLinkPlugin=")]
 	public bool bEnableMetaHumanLiveLinkPlugin = false;
@@ -61,16 +57,9 @@ public class LiveLinkHubTarget : TargetRules
 			"LiveLinkExampleDevice",
 
 			"LiveLinkOpenVR",
-		});
 
-		if (bEnableCaptureManagerPlugin)
-		{
-			OptionalPlugins.AddRange(new string[]
-			{
-				"CaptureManagerApp",
-				"CaptureManagerCore"
-			});
-		}
+			"CaptureManagerApp"
+		});
 
 		if (bEnableMetaHumanLiveLinkPlugin)
 		{
