LiveLinkHub - Enable Aja Blackmagic and Media IO

diff --git a/Engine/Source/Programs/LiveLinkHub/Source/LiveLinkHub.Target.cs b/Engine/Source/Programs/LiveLinkHub/Source/LiveLinkHub.Target.cs
--- a/Engine/Source/Programs/LiveLinkHub/Source/LiveLinkHub.Target.cs
+++ b/Engine/Source/Programs/LiveLinkHub/Source/LiveLinkHub.Target.cs
@@ -59,7 +59,10 @@ public class LiveLinkHubTarget : TargetRules
 			"LiveLinkOpenVR",
 
 			"CaptureManagerApp",
-			"CaptureManagerDevices"
+			"CaptureManagerDevices",
+			"MediaIOFramework",
+			"AjaMedia",
+			"BlackmagicMedia"
 		});
 
 		if (bEnableMetaHumanLiveLinkPlugin)
