Changed the parsing of the command line

* Allows running scripts from PowerShell

diff --git a/Engine/Source/Programs/LiveLinkHub/Source/LiveLinkHubLauncher/Private/Platform/Windows/WindowsLiveLinkHubLauncherMain.cpp b/Engine/Source/Programs/LiveLinkHub/Source/LiveLinkHubLauncher/Private/Platform/Windows/WindowsLiveLinkHubLauncherMain.cpp
--- a/Engine/Source/Programs/LiveLinkHub/Source/LiveLinkHubLauncher/Private/Platform/Windows/WindowsLiveLinkHubLauncherMain.cpp
+++ b/Engine/Source/Programs/LiveLinkHub/Source/LiveLinkHubLauncher/Private/Platform/Windows/WindowsLiveLinkHubLauncherMain.cpp
@@ -8,6 +8,8 @@
 #include "Misc/CommandLine.h"
 #include "Misc/OutputDeviceError.h"
 
+#include <shellapi.h>
+
 
 // Opt in to new D3D12 redist and tell the loader where to search for D3D12Core.dll.
 // The D3D loader looks for these symbol exports in the .exe module.
@@ -24,6 +26,51 @@ extern "C" { _declspec(dllexport) extern const UINT D3D12SDKVersion = 614; } //
 #endif // USE_D3D12_REDIST
 
 
+/**
+ * The command-line invocation string, processed using the standard Windows CommandLineToArgvW implementation.
+ * This need to be a static global to avoid object unwinding errors in WinMain when SEH is enabled.
+ */
+static FString GSavedCommandLine;
+
+// Code copied from Engine\Source\Runtime\Launch\Private\Windows\LaunchWindows.cpp
+bool ProcessCommandLine()
+{
+	int argc = 0;
+	LPWSTR* argv = ::CommandLineToArgvW(::GetCommandLineW(), &argc);
+	if (argv != nullptr)
+	{
+		// Reconstruct our command line string in a format suitable for consumption by the FParse class
+		GSavedCommandLine = "";
+		for (int32 Option = 1; Option < argc; Option++)
+		{
+			GSavedCommandLine += TEXT(" ");
+
+			// Inject quotes in the correct position to preserve arguments containing whitespace
+			FString Temp = argv[Option];
+			if (Temp.Contains(TEXT(" ")))
+			{
+				int32 Quote = 0;
+				if (Temp.StartsWith(TEXT("-")))
+				{
+					int32 Separator;
+					if (Temp.FindChar('=', Separator))
+					{
+						Quote = Separator + 1;
+					}
+				}
+				Temp = Temp.Left(Quote) + TEXT("\"") + Temp.Mid(Quote) + TEXT("\"");
+			}
+			GSavedCommandLine += Temp;
+		}
+
+		// Free memory allocated for CommandLineToArgvW() arguments
+		::LocalFree(argv);
+		return true;
+	}
+
+	return false;
+}
+
 /**
  * The main application entry point for Windows platforms.
  *
@@ -38,7 +85,13 @@ int32 WINAPI WinMain(_In_ HINSTANCE hInInstance, _In_opt_ HINSTANCE hPrevInstanc
 	hInstance = hInInstance;
 
 	const TCHAR* CmdLine = ::GetCommandLineW();
-	CmdLine = FCommandLine::RemoveExeName(CmdLine);
+
+	// Attempt to process the command-line arguments using the standard Windows implementation
+	// (This ensures behavior parity with other platforms where argc and argv are used.)
+	if (ProcessCommandLine())
+	{
+		CmdLine = *GSavedCommandLine;
+	}
 
 #if !UE_BUILD_SHIPPING
 	if (FParse::Param(CmdLine, TEXT("crashreports")))
