Allow LiveLinkHub to save all config sections

diff --git a/Engine/Source/Programs/LiveLinkHub/Config/DefaultEngine.ini b/Engine/Source/Programs/LiveLinkHub/Config/DefaultEngine.ini
--- a/Engine/Source/Programs/LiveLinkHub/Config/DefaultEngine.ini
+++ b/Engine/Source/Programs/LiveLinkHub/Config/DefaultEngine.ini
@@ -1,3 +1,6 @@
+[SectionsToSave]
+bCanSaveAllSections=true
+
 [LiveLink]
 bUseLiveLinkHubSettingsDetailCustomization=true
 bPanelControllerSeparateSourcesSubjects=true
