Abstacting _M_ARM64 and _M_ARM64EC:

_M_ARM64EC   ->   PLATFORM_WINDOWS_ARM64EC
_M_ARM64    ->    PLATFORM_WINDOWS && PLATFORM_CPU_ARM_FAMILY

#rnx

diff --git a/Engine/Plugins/Runtime/MicrosoftSpatialAudio/Source/MicrosoftSpatialSound/Private/MicrosoftSpatialSoundPlugin.cpp b/Engine/Plugins/Runtime/MicrosoftSpatialAudio/Source/MicrosoftSpatialSound/Private/MicrosoftSpatialSoundPlugin.cpp
--- a/Engine/Plugins/Runtime/MicrosoftSpatialAudio/Source/MicrosoftSpatialSound/Private/MicrosoftSpatialSoundPlugin.cpp
+++ b/Engine/Plugins/Runtime/MicrosoftSpatialAudio/Source/MicrosoftSpatialSound/Private/MicrosoftSpatialSoundPlugin.cpp
@@ -399,12 +399,8 @@ uint32 FMicrosoftSpatialSound::Run()
 	return 0;
 }
 
-#if PLATFORM_CPU_ARM_FAMILY
-	#if (defined(__aarch64__) || defined(_M_ARM64))
-		#define TARGET_ARCH TEXT("arm64")
-	#elif defined(_M_ARM64EC)
-		#define TARGET_ARCH TEXT("x64")
-	#endif
+#if PLATFORM_CPU_ARM_FAMILY && !PLATFORM_WINDOWS_ARM64EC
+	#define TARGET_ARCH TEXT("arm64")
 #else
 	#if PLATFORM_64BITS
 		#define TARGET_ARCH TEXT("x64")
diff --git a/Engine/Plugins/XGEController/Source/Private/XGEControllerModule.cpp b/Engine/Plugins/XGEController/Source/Private/XGEControllerModule.cpp
--- a/Engine/Plugins/XGEController/Source/Private/XGEControllerModule.cpp
+++ b/Engine/Plugins/XGEController/Source/Private/XGEControllerModule.cpp
@@ -27,7 +27,7 @@
 	#include "Windows/HideWindowsPlatformTypes.h"
 
 	#define XGE_CONTROL_WORKER_NAME		TEXT("XGEControlWorker")
-#ifdef _M_ARM64EC
+#if PLATFORM_WINDOWS_ARM64EC
 	#define XGE_CONTROL_WORKER_FILENAME	TEXT("XGEControlWorkerarm64ec.exe")
 #elif PLATFORM_CPU_ARM_FAMILY
 	#define XGE_CONTROL_WORKER_FILENAME	TEXT("XGEControlWorkerarm64.exe")
diff --git a/Engine/Source/Developer/ShaderCompilerCommon/Private/DXCWrapper.cpp b/Engine/Source/Developer/ShaderCompilerCommon/Private/DXCWrapper.cpp
--- a/Engine/Source/Developer/ShaderCompilerCommon/Private/DXCWrapper.cpp
+++ b/Engine/Source/Developer/ShaderCompilerCommon/Private/DXCWrapper.cpp
@@ -70,7 +70,7 @@ FDllHandle::FDllHandle(const TCHAR* InFilename)
 {
 #if PLATFORM_WINDOWS
 	check(InFilename && *InFilename);
-#if defined(_M_ARM64)
+#if  PLATFORM_CPU_ARM_FAMILY && !PLATFORM_WINDOWS_ARM64EC
 	FString ShaderConductorDir = FPaths::EngineDir() / TEXT("Binaries/ThirdParty/ShaderConductor/WinArm64");
 #else
 	FString ShaderConductorDir = FPaths::EngineDir() / TEXT("Binaries/ThirdParty/ShaderConductor/Win64");
diff --git a/Engine/Source/Developer/StandaloneRenderer/Private/Windows/D3D/SlateD3DShaders.cpp b/Engine/Source/Developer/StandaloneRenderer/Private/Windows/D3D/SlateD3DShaders.cpp
--- a/Engine/Source/Developer/StandaloneRenderer/Private/Windows/D3D/SlateD3DShaders.cpp
+++ b/Engine/Source/Developer/StandaloneRenderer/Private/Windows/D3D/SlateD3DShaders.cpp
@@ -26,7 +26,7 @@ static HMODULE GetCompilerModule()
 	FString CompilerPath = FPaths::EngineDir();
 #if !PLATFORM_64BITS
 	CompilerPath.Append(TEXT("Binaries/ThirdParty/Windows/DirectX/x86/d3dcompiler_47.dll"));
-#elif defined(_M_ARM64)
+#elif PLATFORM_CPU_ARM_FAMILY && !PLATFORM_WINDOWS_ARM64EC
 	CompilerPath.Append(TEXT("Binaries/ThirdParty/Windows/DirectX/arm64/d3dcompiler_47.dll"));
 #else
 	CompilerPath.Append(TEXT("Binaries/ThirdParty/Windows/DirectX/x64/d3dcompiler_47.dll"));
diff --git a/Engine/Source/Developer/TextureFormatIntelISPCTexComp/Private/TextureFormatIntelISPCTexComp.cpp b/Engine/Source/Developer/TextureFormatIntelISPCTexComp/Private/TextureFormatIntelISPCTexComp.cpp
--- a/Engine/Source/Developer/TextureFormatIntelISPCTexComp/Private/TextureFormatIntelISPCTexComp.cpp
+++ b/Engine/Source/Developer/TextureFormatIntelISPCTexComp/Private/TextureFormatIntelISPCTexComp.cpp
@@ -690,7 +690,7 @@ public:
 
 		FString DLLPath;
 #if PLATFORM_WINDOWS
-	#if defined(_M_ARM64)
+	#if  PLATFORM_CPU_ARM_FAMILY && !PLATFORM_WINDOWS_ARM64EC
 		DLLPath = FPaths::EngineDir() / TEXT("Binaries/ThirdParty/Intel/ISPCTexComp/WinArm64-Release/ispc_texcomp.dll");
 	#else
 		DLLPath = FPaths::EngineDir() / TEXT("Binaries/ThirdParty/Intel/ISPCTexComp/Win64-Release/ispc_texcomp.dll");
diff --git a/Engine/Source/Developer/Windows/ShaderFormatD3D/Private/D3DShaderCompiler.cpp b/Engine/Source/Developer/Windows/ShaderFormatD3D/Private/D3DShaderCompiler.cpp
--- a/Engine/Source/Developer/Windows/ShaderFormatD3D/Private/D3DShaderCompiler.cpp
+++ b/Engine/Source/Developer/Windows/ShaderFormatD3D/Private/D3DShaderCompiler.cpp
@@ -393,7 +393,7 @@ public:
 private:
 	FxcCompilerFunctions()
 	{
-#if defined(_M_ARM64)
+#if  PLATFORM_CPU_ARM_FAMILY && !PLATFORM_WINDOWS_ARM64EC
 		FString CompilerPath = FPaths::EngineDir() / TEXT("Binaries/ThirdParty/Windows/DirectX/arm64/d3dcompiler_47.dll");
 #else
 		FString CompilerPath = FPaths::EngineDir() / TEXT("Binaries/ThirdParty/Windows/DirectX/x64/d3dcompiler_47.dll");
diff --git a/Engine/Source/Developer/Windows/WindowsTargetPlatfomControls/Public/GenericWindowsTargetPlatformControls.h b/Engine/Source/Developer/Windows/WindowsTargetPlatfomControls/Public/GenericWindowsTargetPlatformControls.h
--- a/Engine/Source/Developer/Windows/WindowsTargetPlatfomControls/Public/GenericWindowsTargetPlatformControls.h
+++ b/Engine/Source/Developer/Windows/WindowsTargetPlatfomControls/Public/GenericWindowsTargetPlatformControls.h
@@ -115,7 +115,7 @@ public:
 
 	virtual void GetShaderCompilerDependencies(TArray<FString>& OutDependencies) const override
 	{
-#if defined(_M_ARM64)
+#if  PLATFORM_CPU_ARM_FAMILY && !PLATFORM_WINDOWS_ARM64EC
 		FTargetPlatformControlsBase::AddDependencySCArrayHelper(OutDependencies, TEXT("Binaries/ThirdParty/Windows/DirectX/arm64/d3dcompiler_47.dll"));
 		FTargetPlatformControlsBase::AddDependencySCArrayHelper(OutDependencies, TEXT("Binaries/ThirdParty/ShaderConductor/WinArm64/ShaderConductor.dll"));
 		FTargetPlatformControlsBase::AddDependencySCArrayHelper(OutDependencies, TEXT("Binaries/ThirdParty/ShaderConductor/WinArm64/dxcompiler.dll"));
diff --git a/Engine/Source/Programs/ChaosVisualDebugger/Private/ChaosVisualDebuggerMain.cpp b/Engine/Source/Programs/ChaosVisualDebugger/Private/ChaosVisualDebuggerMain.cpp
--- a/Engine/Source/Programs/ChaosVisualDebugger/Private/ChaosVisualDebuggerMain.cpp
+++ b/Engine/Source/Programs/ChaosVisualDebugger/Private/ChaosVisualDebuggerMain.cpp
@@ -17,11 +17,11 @@ DEFINE_LOG_CATEGORY_STATIC(LogChaosVisualDebugger, Log, All);
 // versions of Windows will transparently load default OS-provided D3D12 library.
 #if USE_D3D12_REDIST
 extern "C" { _declspec(dllexport) extern const UINT D3D12SDKVersion = 614; } // D3D12_SDK_VERSION
-#if defined(_M_ARM64)
+#if PLATFORM_CPU_ARM_FAMILY && !PLATFORM_WINDOWS_ARM64EC
 	extern "C" { _declspec(dllexport) extern const char* D3D12SDKPath = ".\\D3D12\\arm64\\"; }
 #else
 	extern "C" { _declspec(dllexport) extern const char* D3D12SDKPath = ".\\D3D12\\x64\\"; }
-#endif // _M_ARM64
+#endif // PLATFORM_CPU_ARM_FAMILY && !PLATFORM_WINDOWS_ARM64EC
 #endif // USE_D3D12_REDIST
 
 int32 RunChaosVisualDebugger(const TCHAR* CommandLine)
diff --git a/Engine/Source/Programs/LiveLinkHub/Source/LiveLinkHubLauncher/Private/Platform/Windows/WindowsLiveLinkHubLauncherMain.cpp b/Engine/Source/Programs/LiveLinkHub/Source/LiveLinkHubLauncher/Private/Platform/Windows/WindowsLiveLinkHubLauncherMain.cpp
--- a/Engine/Source/Programs/LiveLinkHub/Source/LiveLinkHubLauncher/Private/Platform/Windows/WindowsLiveLinkHubLauncherMain.cpp
+++ b/Engine/Source/Programs/LiveLinkHub/Source/LiveLinkHubLauncher/Private/Platform/Windows/WindowsLiveLinkHubLauncherMain.cpp
@@ -18,11 +18,11 @@
 #define USE_D3D12_REDIST (PLATFORM_DESKTOP && PLATFORM_64BITS && 1)
 #if USE_D3D12_REDIST
 extern "C" { _declspec(dllexport) extern const UINT D3D12SDKVersion = 614; } // D3D12_SDK_VERSION
-#if defined(_M_ARM64)
+#if PLATFORM_CPU_ARM_FAMILY && !PLATFORM_WINDOWS_ARM64EC
 	extern "C" { _declspec(dllexport) extern const char* D3D12SDKPath = ".\\D3D12\\arm64\\"; }
 #else
 	extern "C" { _declspec(dllexport) extern const char* D3D12SDKPath = ".\\D3D12\\x64\\"; }
-#endif // _M_ARM64
+#endif // PLATFORM_CPU_ARM_FAMILY && !PLATFORM_WINDOWS_ARM64EC
 #endif // USE_D3D12_REDIST
 
 
diff --git a/Engine/Source/Runtime/CEF3Utils/Private/CEF3Utils.cpp b/Engine/Source/Runtime/CEF3Utils/Private/CEF3Utils.cpp
--- a/Engine/Source/Runtime/CEF3Utils/Private/CEF3Utils.cpp
+++ b/Engine/Source/Runtime/CEF3Utils/Private/CEF3Utils.cpp
@@ -70,7 +70,7 @@ namespace CEF3Utils
 	{
 #if PLATFORM_WINDOWS
 	#if PLATFORM_64BITS
-		#if PLATFORM_CPU_ARM_FAMILY && !defined(_M_ARM64EC)
+		#if PLATFORM_CPU_ARM_FAMILY && !PLATFORM_WINDOWS_ARM64EC
 			DllPath = FPaths::Combine(*FPaths::EngineDir(), TEXT("Binaries/ThirdParty/CEF3/WinArm64"), CEF_VERSION_DIR);
 		#else
 			DllPath = FPaths::Combine(*FPaths::EngineDir(), TEXT("Binaries/ThirdParty/CEF3/Win64"), CEF_VERSION_DIR);
diff --git a/Engine/Source/Runtime/Core/Private/Microsoft/MicrosoftPlatformStackWalk.cpp b/Engine/Source/Runtime/Core/Private/Microsoft/MicrosoftPlatformStackWalk.cpp
--- a/Engine/Source/Runtime/Core/Private/Microsoft/MicrosoftPlatformStackWalk.cpp
+++ b/Engine/Source/Runtime/Core/Private/Microsoft/MicrosoftPlatformStackWalk.cpp
@@ -3,6 +3,7 @@
 #include "Microsoft/MicrosoftPlatformStackWalk.h"
 
 #include "Misc/Paths.h"
+#include "CoreGlobals.h"
 
 #include "Microsoft/AllowMicrosoftPlatformTypes.h"
 	#include <psapi.h>
@@ -78,7 +79,7 @@ bool FMicrosoftPlatformStackWalk::ExtractInfoFromModule(void* ProcessHandle, voi
 
 int32 FMicrosoftPlatformStackWalk::CaptureStackTraceInternal(uint64* OutBacktrace, uint32 MaxDepth, void* InContext, void* InThreadHandle, uint32* OutDepth)
 {
-#if PLATFORM_CPU_X86_FAMILY || defined(_M_ARM64EC)
+#if PLATFORM_CPU_X86_FAMILY || PLATFORM_WINDOWS_ARM64EC
 	uint32 CurrentDepth = 0;
 	HANDLE ThreadHandle = reinterpret_cast<HANDLE>(InThreadHandle);
 	CONTEXT ContextCopy = *reinterpret_cast<PCONTEXT>(InContext);
diff --git a/Engine/Source/Runtime/Core/Private/ProfilingDebugging/Windows/PIXExternalProfiler.cpp b/Engine/Source/Runtime/Core/Private/ProfilingDebugging/Windows/PIXExternalProfiler.cpp
--- a/Engine/Source/Runtime/Core/Private/ProfilingDebugging/Windows/PIXExternalProfiler.cpp
+++ b/Engine/Source/Runtime/Core/Private/ProfilingDebugging/Windows/PIXExternalProfiler.cpp
@@ -110,7 +110,7 @@ public:
 	bool Initialize()
 	{
 #if USE_PIX
-#if defined(_M_ARM64)
+#if PLATFORM_CPU_ARM_FAMILY && !PLATFORM_WINDOWS_ARM64EC
 		FString WindowsPixDllRelativePath = FPaths::Combine(*FPaths::EngineDir(), TEXT("Binaries/ThirdParty/Windows/WinPixEventRuntime/arm64"));
 		FString WindowsPixDll("WinPixEventRuntime.dll");
 #else
diff --git a/Engine/Source/Runtime/Core/Private/Windows/WindowsPlatformCrashContext.cpp b/Engine/Source/Runtime/Core/Private/Windows/WindowsPlatformCrashContext.cpp
--- a/Engine/Source/Runtime/Core/Private/Windows/WindowsPlatformCrashContext.cpp
+++ b/Engine/Source/Runtime/Core/Private/Windows/WindowsPlatformCrashContext.cpp
@@ -1877,7 +1877,7 @@ static void ReportEvent(ECrashContextType InType, const TCHAR* ErrorMessage, uin
 
 					if (ProgramCounter == nullptr)
 					{
-#if PLATFORM_CPU_X86_FAMILY || defined(_M_ARM64EC)
+#if PLATFORM_CPU_X86_FAMILY || PLATFORM_WINDOWS_ARM64EC
 						ProgramCounter = (void*)(ThreadContext.Rip);
 #else
 						ProgramCounter = (void*)(ThreadContext.Pc);
diff --git a/Engine/Source/Runtime/Core/Private/Windows/WindowsPlatformMisc.cpp b/Engine/Source/Runtime/Core/Private/Windows/WindowsPlatformMisc.cpp
--- a/Engine/Source/Runtime/Core/Private/Windows/WindowsPlatformMisc.cpp
+++ b/Engine/Source/Runtime/Core/Private/Windows/WindowsPlatformMisc.cpp
@@ -4262,7 +4262,7 @@ FString FWindowsPlatformMisc::GetOperatingSystemId()
 
 FString FWindowsPlatformMisc::GetHostArchitecture()
 {
-#if defined(_M_ARM64EC) || defined(_M_ARM64)
+#if PLATFORM_CPU_ARM_FAMILY
 
 	// these build targets can only run on arm64 (assuming there's no emulation)
 	return TEXT("arm64");
diff --git a/Engine/Source/Runtime/Core/Private/Windows/WindowsPlatformProcess.cpp b/Engine/Source/Runtime/Core/Private/Windows/WindowsPlatformProcess.cpp
--- a/Engine/Source/Runtime/Core/Private/Windows/WindowsPlatformProcess.cpp
+++ b/Engine/Source/Runtime/Core/Private/Windows/WindowsPlatformProcess.cpp
@@ -271,9 +271,9 @@ FString FWindowsPlatformProcess::GenerateApplicationPath( const FString& AppName
 
 const TCHAR* FWindowsPlatformProcess::GetArchitectureSuffix()
 {
-#if defined(_M_ARM64EC)
+#if PLATFORM_WINDOWS_ARM64EC
 	return TEXT("arm64ec");
-#elif defined(_M_ARM64)
+#elif PLATFORM_CPU_ARM_FAMILY
 	return TEXT("arm64");
 #else
 	return nullptr;
diff --git a/Engine/Source/Runtime/Core/Private/Windows/WindowsPlatformStackWalk.cpp b/Engine/Source/Runtime/Core/Private/Windows/WindowsPlatformStackWalk.cpp
--- a/Engine/Source/Runtime/Core/Private/Windows/WindowsPlatformStackWalk.cpp
+++ b/Engine/Source/Runtime/Core/Private/Windows/WindowsPlatformStackWalk.cpp
@@ -217,7 +217,7 @@ static int32 CaptureStackTraceExternalProcess(uint64* BackTrace, uint32 MaxDepth
 		StackFrame64.AddrStack.Mode      = AddrModeFlat;
 		StackFrame64.AddrFrame.Mode      = AddrModeFlat;
 #if PLATFORM_64BITS
-#if defined(_M_ARM64)
+#if PLATFORM_CPU_ARM_FAMILY && !PLATFORM_WINDOWS_ARM64EC
 		StackFrame64.AddrPC.Offset       = ContextCopy.Pc;
 		StackFrame64.AddrStack.Offset    = ContextCopy.Sp;
 		StackFrame64.AddrFrame.Offset    = ContextCopy.Fp;
diff --git a/Engine/Source/Runtime/Core/Public/GenericPlatform/GenericPlatformMath.h b/Engine/Source/Runtime/Core/Public/GenericPlatform/GenericPlatformMath.h
--- a/Engine/Source/Runtime/Core/Public/GenericPlatform/GenericPlatformMath.h
+++ b/Engine/Source/Runtime/Core/Public/GenericPlatform/GenericPlatformMath.h
@@ -481,7 +481,7 @@ struct FGenericPlatformMath
 	}
 
 	// Returns e^Value
-#if PLATFORM_WINDOWS && (defined(_M_ARM64) || defined(_M_ARM64EC))
+#if PLATFORM_WINDOWS && PLATFORM_CPU_ARM_FAMILY
 	[[nodiscard]] static FORCEINLINE float Exp( volatile float Value ) { return expf(Value); } // #jira UE-261901 : workaround optimization bug for expf()
 #else
 	[[nodiscard]] static FORCEINLINE float Exp( float Value ) { return expf(Value); }
diff --git a/Engine/Source/Runtime/Core/Public/GenericPlatform/GenericPlatformProcess.h b/Engine/Source/Runtime/Core/Public/GenericPlatform/GenericPlatformProcess.h
--- a/Engine/Source/Runtime/Core/Public/GenericPlatform/GenericPlatformProcess.h
+++ b/Engine/Source/Runtime/Core/Public/GenericPlatform/GenericPlatformProcess.h
@@ -28,7 +28,7 @@ class FEvent;
 #	include <intrin.h>
 #	if defined(_M_ARM)
 #		include <armintr.h>
-#	elif defined(_M_ARM64) || defined(_M_ARM64EC)
+#	elif PLATFORM_WINDOWS && PLATFORM_CPU_ARM_FAMILY
 #		include <arm64intr.h>
 #	endif
 #endif
diff --git a/Engine/Source/Runtime/Core/Public/HAL/Platform.h b/Engine/Source/Runtime/Core/Public/HAL/Platform.h
--- a/Engine/Source/Runtime/Core/Public/HAL/Platform.h
+++ b/Engine/Source/Runtime/Core/Public/HAL/Platform.h
@@ -92,6 +92,13 @@
 	#endif
 #endif
 
+// Whether the platform is Arm64ec. Normal Windows Arm64 can be detected with PLATFORM_CPU_ARM_FAMILY && PLATFORM_WINDOWS
+#if defined(_M_ARM64EC)
+	#define PLATFORM_WINDOWS_ARM64EC PLATFORM_WINDOWS
+#else
+	#define PLATFORM_WINDOWS_ARM64EC 0
+#endif
+
 #if PLATFORM_APPLE
 	#include <stddef.h> // needed for size_t
 #endif
diff --git a/Engine/Source/Runtime/Core/Public/Math/UnrealMathNeon.h b/Engine/Source/Runtime/Core/Public/Math/UnrealMathNeon.h
--- a/Engine/Source/Runtime/Core/Public/Math/UnrealMathNeon.h
+++ b/Engine/Source/Runtime/Core/Public/Math/UnrealMathNeon.h
@@ -2247,7 +2247,7 @@ FORCEINLINE int32 VectorAnyGreaterThan(VectorRegister4Double Vec1, VectorRegiste
  * https://learn.microsoft.com/en-us/cpp/build/arm64ec-windows-abi-conventions
  * "...Avoid directly reading, writing, or computing mappings of FPCR and FPSR. These bits may be used in the future and are subject to change."
  */
-#if PLATFORM_WINDOWS && defined(_M_ARM64EC)
+#if PLATFORM_WINDOWS_ARM64EC
 	#pragma warning(push)
 	#pragma warning(disable:5076) // warning C5076: read from FPCR
 	#pragma warning(disable:5077) // warning C5076: write to FPCR
@@ -2289,7 +2289,7 @@ FORCEINLINE void VectorSetControlRegister(uint32_t ControlStatus)
 /**
  * Restore arm64ec FPCR warnings
  */
-#if PLATFORM_WINDOWS && defined(_M_ARM64EC)
+#if PLATFORM_WINDOWS_ARM64EC
 	#pragma warning(pop)
 #endif
 
diff --git a/Engine/Source/Runtime/Core/Public/Windows/AllowWindowsPlatformAtomics.h b/Engine/Source/Runtime/Core/Public/Windows/AllowWindowsPlatformAtomics.h
--- a/Engine/Source/Runtime/Core/Public/Windows/AllowWindowsPlatformAtomics.h
+++ b/Engine/Source/Runtime/Core/Public/Windows/AllowWindowsPlatformAtomics.h
@@ -14,7 +14,7 @@
 
 #define InterlockedIncrement _InterlockedIncrement
 #define InterlockedDecrement _InterlockedDecrement
-#if !defined(InterlockedAdd) && !defined(_M_ARM64EC)
+#if !defined(InterlockedAdd) && !PLATFORM_WINDOWS_ARM64EC
 	#define InterlockedAdd _InterlockedAdd
 #endif
 #define InterlockedExchange _InterlockedExchange
diff --git a/Engine/Source/Runtime/D3D12RHI/Private/Windows/WindowsD3D12Device.cpp b/Engine/Source/Runtime/D3D12RHI/Private/Windows/WindowsD3D12Device.cpp
--- a/Engine/Source/Runtime/D3D12RHI/Private/Windows/WindowsD3D12Device.cpp
+++ b/Engine/Source/Runtime/D3D12RHI/Private/Windows/WindowsD3D12Device.cpp
@@ -1116,7 +1116,7 @@ FDynamicRHI* FD3D12DynamicRHIModule::CreateRHI(ERHIFeatureLevel::Type RequestedF
 void FD3D12DynamicRHIModule::StartupModule()
 {
 #if D3D12RHI_SUPPORTS_WIN_PIX
-#if defined(_M_ARM64)
+#if PLATFORM_CPU_ARM_FAMILY && !PLATFORM_WINDOWS_ARM64EC
 	static FString WindowsPixDllRelativePath = FPaths::Combine(FPaths::EngineDir(), TEXT("Binaries/ThirdParty/Windows/WinPixEventRuntime/arm64"));
 	static const TCHAR* WindowsPixDll = TEXT("WinPixEventRuntime.dll");
 #else
diff --git a/Engine/Source/Runtime/D3D12RHI/Private/Windows/WindowsD3D12PipelineState.cpp b/Engine/Source/Runtime/D3D12RHI/Private/Windows/WindowsD3D12PipelineState.cpp
--- a/Engine/Source/Runtime/D3D12RHI/Private/Windows/WindowsD3D12PipelineState.cpp
+++ b/Engine/Source/Runtime/D3D12RHI/Private/Windows/WindowsD3D12PipelineState.cpp
@@ -597,7 +597,7 @@ static void DumpShaderAsm(FString& String, const D3D12_SHADER_BYTECODE& Shader)
 		if (D3DDisasmFunc == nullptr)
 		{
 			static HMODULE CompilerDLL = NULL;	// not nullptr as the return value of LoadLibrary remains defined as NULL
-#if defined(_M_ARM64)
+#if PLATFORM_CPU_ARM_FAMILY && !PLATFORM_WINDOWS_ARM64EC
 			static const TCHAR* CompilerPath = TEXT("Binaries/ThirdParty/Windows/DirectX/arm64/d3dcompiler_47.dll");
 #else
 			static const TCHAR* CompilerPath = TEXT("Binaries/ThirdParty/Windows/DirectX/x64/d3dcompiler_47.dll");
diff --git a/Engine/Source/Runtime/Engine/Private/ShaderCompiler/ShaderCompiler.cpp b/Engine/Source/Runtime/Engine/Private/ShaderCompiler/ShaderCompiler.cpp
--- a/Engine/Source/Runtime/Engine/Private/ShaderCompiler/ShaderCompiler.cpp
+++ b/Engine/Source/Runtime/Engine/Private/ShaderCompiler/ShaderCompiler.cpp
@@ -728,7 +728,7 @@ FShaderCompilingManager::FShaderCompilingManager() :
 	WorkersBusyTime = 0;
 
 #if PLATFORM_WINDOWS
-#ifdef _M_ARM64EC
+#ifdef PLATFORM_WINDOWS_ARM64EC
 	FString ExecutableName("ShaderCompileWorkerarm64ec.exe");
 #elif PLATFORM_CPU_ARM_FAMILY
 	FString ExecutableName("ShaderCompileWorkerarm64.exe");
diff --git a/Engine/Source/Runtime/Launch/Private/Windows/LaunchWindows.cpp b/Engine/Source/Runtime/Launch/Private/Windows/LaunchWindows.cpp
--- a/Engine/Source/Runtime/Launch/Private/Windows/LaunchWindows.cpp
+++ b/Engine/Source/Runtime/Launch/Private/Windows/LaunchWindows.cpp
@@ -54,11 +54,11 @@ extern "C" { _declspec(dllexport) uint32 AmdPowerXpressRequestHighPerformance =
 #define USE_D3D12_REDIST (PLATFORM_DESKTOP && PLATFORM_64BITS && 1)
 #if USE_D3D12_REDIST
 extern "C" { _declspec(dllexport) extern const UINT D3D12SDKVersion = 614; } // D3D12_SDK_VERSION
-#if defined(_M_ARM64)
+#if PLATFORM_CPU_ARM_FAMILY && !PLATFORM_WINDOWS_ARM64EC
 	extern "C" { _declspec(dllexport) extern const char* D3D12SDKPath = ".\\D3D12\\arm64\\"; }
 #else
 	extern "C" { _declspec(dllexport) extern const char* D3D12SDKPath = ".\\D3D12\\x64\\"; }
-#endif // _M_ARM64
+#endif // PLATFORM_CPU_ARM_FAMILY && !PLATFORM_WINDOWS_ARM64EC
 #endif // USE_D3D12_REDIST
 
 /** Whether we should pause before exiting. used by UCC */
diff --git a/Engine/Source/Runtime/VorbisAudioDecoder/Module/Private/VorbisAudioInfo.cpp b/Engine/Source/Runtime/VorbisAudioDecoder/Module/Private/VorbisAudioInfo.cpp
--- a/Engine/Source/Runtime/VorbisAudioDecoder/Module/Private/VorbisAudioInfo.cpp
+++ b/Engine/Source/Runtime/VorbisAudioDecoder/Module/Private/VorbisAudioInfo.cpp
@@ -809,7 +809,7 @@ void LoadVorbisLibraries()
 #endif
 
 		// vorbis winarm64 libs are static for now, so no need to load anything
-#if PLATFORM_CPU_ARM_FAMILY && !defined(_M_ARM64EC)
+#if PLATFORM_CPU_ARM_FAMILY && !PLATFORM_WINDOWS_ARM64EC
 		FString RootOggPath = FPaths::EngineDir() / TEXT("Binaries/ThirdParty/Ogg/") / PlatformString / VSVersion / TEXT("arm64/");
 		FString DLLToLoad = RootOggPath + TEXT("libogg") + DLLNameStub;
 		void* LibOggHandle = FPlatformProcess::GetDllHandle(*DLLToLoad);
diff --git a/Engine/Source/Runtime/WebBrowser/Private/WebBrowserSingleton.cpp b/Engine/Source/Runtime/WebBrowser/Private/WebBrowserSingleton.cpp
--- a/Engine/Source/Runtime/WebBrowser/Private/WebBrowserSingleton.cpp
+++ b/Engine/Source/Runtime/WebBrowser/Private/WebBrowserSingleton.cpp
@@ -67,7 +67,7 @@ THIRD_PARTY_INCLUDES_END
 #	endif
 #	define CEF3_BIN_DIR TEXT("Binaries/ThirdParty/CEF3")
 #	if PLATFORM_WINDOWS && PLATFORM_64BITS && PLATFORM_CPU_ARM_FAMILY
-#		if defined(_M_ARM64EC)
+#		if PLATFORM_WINDOWS_ARM64EC
 #			define CEF3_RESOURCES_DIR CEF3_BIN_DIR TEXT("/Win64") CEF_VERSION_DIR TEXT("/Resources")
 #		else
 #			define CEF3_RESOURCES_DIR CEF3_BIN_DIR TEXT("/WinArm64") CEF_VERSION_DIR TEXT("/Resources")
