Live Link Hub: Ensure MetaHumanCoreTechLib plugin content is included in packaged builds.

#rnx

diff --git a/Engine/Source/Programs/LiveLinkHub/Source/LiveLinkHubLauncher/RuntimeDependencies.ini b/Engine/Source/Programs/LiveLinkHub/Source/LiveLinkHubLauncher/RuntimeDependencies.ini
--- a/Engine/Source/Programs/LiveLinkHub/Source/LiveLinkHubLauncher/RuntimeDependencies.ini
+++ b/Engine/Source/Programs/LiveLinkHub/Source/LiveLinkHubLauncher/RuntimeDependencies.ini
@@ -92,9 +92,9 @@
 +RuntimeDependencies="/Content/VREditor/ViewportInteractionAssetContainerData.uasset"
 
 ; Plugin content
-; FIXME: ACLPlugin still generates warnings at runtime
 +RuntimeDependencies="/Plugins/Animation/ACLPlugin/Content/..."
 +RuntimeDependencies="/Plugins/Animation/LiveLink/Content/..."
++RuntimeDependencies="/Plugins/MetaHuman/MetaHumanCoreTechLib/Content/..."
 +RuntimeDependencies="/Plugins/MovieScene/LevelSequenceEditor/Content/..."
 +RuntimeDependencies="/Plugins/Runtime/AR/ARUtilities/Content/..."
 +RuntimeDependencies="/Plugins/Runtime/AR/AppleAR/AppleARKit/Content/..."
