LiveLinkHub - Disable automatic ARKit source creation by default
Fix outbound name having disappeared from LiveLinkHub

diff --git a/Engine/Source/Programs/LiveLinkHub/Config/DefaultEngine.ini b/Engine/Source/Programs/LiveLinkHub/Config/DefaultEngine.ini
--- a/Engine/Source/Programs/LiveLinkHub/Config/DefaultEngine.ini
+++ b/Engine/Source/Programs/LiveLinkHub/Config/DefaultEngine.ini
@@ -30,9 +30,10 @@ AutoConnectClients=LocalOnly
 +ProjectSettingsToDisplay=Timing & Sync
 +ProjectSettingsToDisplay=UDPMessaging
 bShowFrameRate=True
+TargetFrameRate=60.000000
 
 [/Script/AppleARKit.AppleARKitSettings]
-bEnableLiveLinkForFaceTracking=true
+bEnableLiveLinkForFaceTracking=false
 
 [/Script/UnrealEd.ProjectPackagingSettings]
 bShareMaterialShaderCode=false
diff --git a/Engine/Source/Programs/LiveLinkHub/Config/DefaultGame.ini b/Engine/Source/Programs/LiveLinkHub/Config/DefaultGame.ini
--- a/Engine/Source/Programs/LiveLinkHub/Config/DefaultGame.ini
+++ b/Engine/Source/Programs/LiveLinkHub/Config/DefaultGame.ini
@@ -11,10 +11,11 @@ bIsForExternalDistribution=true
 
 [/Script/LiveLink.LiveLinkSettings]
 DefaultSettingsClass=/Script/LiveLink.LiveLinkHubSubjectSettings
-bPreProcessRebroadcastFrames=true
-bTranslateRebroadcastFrames=true
+-DefaultRoleSettings=(Role="/Script/CoreUObject.Class'/Script/LiveLinkInterface.LiveLinkAnimationRole'",SettingClass="/Script/CoreUObject.Class'/Script/LiveLinkInterface.LiveLinkSubjectSettings'",FrameInterpolationProcessor="/Script/CoreUObject.Class'/Script/LiveLink.LiveLinkAnimationFrameInterpolationProcessor'",FramePreProcessors=)
++DefaultRoleSettings=(Role="/Script/CoreUObject.Class'/Script/LiveLinkInterface.LiveLinkAnimationRole'",SettingClass="/Script/CoreUObject.Class'/Script/LiveLink.LiveLinkHubSubjectSettings'",FrameInterpolationProcessor="/Script/CoreUObject.Class'/Script/LiveLink.LiveLinkAnimationFrameInterpolationProcessor'",FramePreProcessors=)
 
 
 [/Script/EngineSettings.GeneralProjectSettings]
 ProjectID=914C155297C9471A9E34100517697A75
 ProjectName=LiveLinkHub
+
