LiveLinkHub - Enable preprocessing and translating rebroadcast frames

diff --git a/Engine/Source/Programs/LiveLinkHub/Config/DefaultGame.ini b/Engine/Source/Programs/LiveLinkHub/Config/DefaultGame.ini
--- a/Engine/Source/Programs/LiveLinkHub/Config/DefaultGame.ini
+++ b/Engine/Source/Programs/LiveLinkHub/Config/DefaultGame.ini
@@ -13,7 +13,8 @@ bIsForExternalDistribution=true
 DefaultSettingsClass=/Script/LiveLink.LiveLinkHubSubjectSettings
 -DefaultRoleSettings=(Role="/Script/CoreUObject.Class'/Script/LiveLinkInterface.LiveLinkAnimationRole'",SettingClass="/Script/CoreUObject.Class'/Script/LiveLinkInterface.LiveLinkSubjectSettings'",FrameInterpolationProcessor="/Script/CoreUObject.Class'/Script/LiveLink.LiveLinkAnimationFrameInterpolationProcessor'",FramePreProcessors=)
 +DefaultRoleSettings=(Role="/Script/CoreUObject.Class'/Script/LiveLinkInterface.LiveLinkAnimationRole'",SettingClass="/Script/CoreUObject.Class'/Script/LiveLink.LiveLinkHubSubjectSettings'",FrameInterpolationProcessor="/Script/CoreUObject.Class'/Script/LiveLink.LiveLinkAnimationFrameInterpolationProcessor'",FramePreProcessors=)
-
+bPreProcessRebroadcastFrames=true
+bTranslateRebroadcastFrames=true
 
 [/Script/EngineSettings.GeneralProjectSettings]
 ProjectID=914C155297C9471A9E34100517697A75
