Live Link Hub: Store external plugin paths in a versioned directory that reflects the ABI. If we're not an installed build, use the workspace Engine/Config dir instead.

diff --git a/Engine/Source/Programs/LiveLinkHub/Source/LiveLinkHubLauncher/Private/LiveLinkHubPluginHelpers.cpp b/Engine/Source/Programs/LiveLinkHub/Source/LiveLinkHubLauncher/Private/LiveLinkHubPluginHelpers.cpp
--- a/Engine/Source/Programs/LiveLinkHub/Source/LiveLinkHubLauncher/Private/LiveLinkHubPluginHelpers.cpp
+++ b/Engine/Source/Programs/LiveLinkHub/Source/LiveLinkHubLauncher/Private/LiveLinkHubPluginHelpers.cpp
@@ -6,6 +6,7 @@
 #include "IPluginBrowser.h"
 #include "Misc/App.h"
 #include "Misc/ConfigCacheIni.h"
+#include "Misc/EngineVersion.h"
 
 
 namespace LiveLinkHub::PluginHelpers
@@ -17,13 +18,26 @@ namespace LiveLinkHub::PluginHelpers
 
 FString LiveLinkHub::PluginHelpers::GetPluginConfigPath()
 {
-	return FPaths::Combine(
-		FPlatformProcess::UserSettingsDir(),
-		*FApp::GetEpicProductIdentifier(),
-		TEXT("LiveLinkHub"),
-		TEXT("Config"),
-		TEXT("LiveLinkHubPlugins.ini")
-	);
+	// Modeled after FPaths::EngineUserDir()
+	if (FPaths::ShouldSaveToUserDir() || FApp::IsEngineInstalled())
+	{
+		// e.g. %LocalAppData%\UnrealEngine\LiveLinkHub_5.7\Config\LiveLinkHubPlugins.ini
+		return FPaths::Combine(
+			FPlatformProcess::UserSettingsDir(),
+			*FApp::GetEpicProductIdentifier(),
+			*FString(TEXT("LiveLinkHub_")) + FEngineVersion::Current().ToString(EVersionComponent::Minor),
+			TEXT("Config"),
+			TEXT("LiveLinkHubPlugins.ini")
+		);
+	}
+	else
+	{
+		return FPaths::Combine(
+			*FPaths::EngineConfigDir(),
+			TEXT("LiveLinkHub"),
+			TEXT("LiveLinkHubPlugins.ini")
+		);
+	}
 }
 
 void LiveLinkHub::PluginHelpers::ReadPluginDirectoriesFromConfig(TArray<FString>& OutDirectories)
