LiveLinkHub - Set mininum driver version to D3D12 driver that fixes memory leak issue

diff --git a/Engine/Source/Programs/LiveLinkHub/Config/DefaultHardware.ini b/Engine/Source/Programs/LiveLinkHub/Config/DefaultHardware.ini
new file mode 100644
--- /dev/null
+++ b/Engine/Source/Programs/LiveLinkHub/Config/DefaultHardware.ini
@@ -0,0 +1,3 @@
+[GPU_NVIDIA Windows]
+SuggestedDriverVersion="581.08"
++DriverDenyList=(DriverVersion="<581.08", Reason="These driver versions have a known stability issue that causes a memory leak in Live Link Hub.", RHI="D3D12")
\ No newline at end of file
