LiveLinkHub - Fix mac build not working because of project generation and project path not resolving correctly.
Fix crash on startup caused by oodle texture format not being enabled

diff --git a/Engine/Source/Programs/LiveLinkHub/Config/DefaultEngine.ini b/Engine/Source/Programs/LiveLinkHub/Config/DefaultEngine.ini
--- a/Engine/Source/Programs/LiveLinkHub/Config/DefaultEngine.ini
+++ b/Engine/Source/Programs/LiveLinkHub/Config/DefaultEngine.ini
@@ -48,10 +48,6 @@ DefaultEditorTarget=LiveLinkHub
 ; Disabling this to ensure that DesktopPlatformBase doesn't generate a TargetInfo.json file with UE's hundreds of targets
 bIncludeParentAssembly=false
 
-; Not needed at the moment, can be removed if we enable the TextureFormatOodle plugin
-[AlternateTextureCompression]
-TextureCompressionFormat=""
-
 [/Script/InterchangeEngine.InterchangeProjectSettings]
 ContentImportSettings=()
 SceneImportSettings=()
diff --git a/Engine/Source/Programs/LiveLinkHub/Source/LiveLinkHub.Target.cs b/Engine/Source/Programs/LiveLinkHub/Source/LiveLinkHub.Target.cs
--- a/Engine/Source/Programs/LiveLinkHub/Source/LiveLinkHub.Target.cs
+++ b/Engine/Source/Programs/LiveLinkHub/Source/LiveLinkHub.Target.cs
@@ -29,7 +29,9 @@ public class LiveLinkHubTarget : TargetRules
 
 		DefaultBuildSettings = BuildSettingsVersion.Latest;
 		IncludeOrderVersion = EngineIncludeOrderVersion.Latest;
-
+		
+		MacPlatform.bForceTargetNameAsExeName = true;
+		
 		// These plugins are required for running LiveLinkHub. 
 		// They may be a direct dependency or a dependency of another plugin.
 		EnablePlugins.AddRange(new string[]
@@ -59,6 +61,7 @@ public class LiveLinkHubTarget : TargetRules
 			"LiveLinkOpenVR",
 
 			"LiveLinkOpenTrackIO",
+			"TextureFormatOodle",
 
 			"CaptureManagerApp",
 			"CaptureManagerDevices",
diff --git a/Engine/Source/Programs/LiveLinkHub/Source/LiveLinkHubLauncher/Private/LiveLinkHubRun.cpp b/Engine/Source/Programs/LiveLinkHub/Source/LiveLinkHubLauncher/Private/LiveLinkHubRun.cpp
--- a/Engine/Source/Programs/LiveLinkHub/Source/LiveLinkHubLauncher/Private/LiveLinkHubRun.cpp
+++ b/Engine/Source/Programs/LiveLinkHub/Source/LiveLinkHubLauncher/Private/LiveLinkHubRun.cpp
@@ -55,14 +55,20 @@ int32 RunLiveLinkHub(const TCHAR* CommandLine)
 #endif
 
 	bool bPendingRestart = false;
-
-	const TCHAR* const DevelopmentProjectPath = TEXT("../../Source/Programs/LiveLinkHub/LiveLinkHub.uproject");
-	const TCHAR* const StagedProjectPath = TEXT("../../../LiveLinkHub/LiveLinkHub.uproject");
-
-	if (FPaths::FileExists(DevelopmentProjectPath))
+	
+	FString DevelopmentProjectPath = TEXT("../../Source/Programs/LiveLinkHub/LiveLinkHub.uproject");
+	FString StagedProjectPath = TEXT("../../../LiveLinkHub/LiveLinkHub.uproject");
+
+#if PLATFORM_MAC
+	// On Mac the following FileExists calls will fail unless it's an absolute path.
+	DevelopmentProjectPath = FPaths::ConvertRelativePathToFull(DevelopmentProjectPath);
+	StagedProjectPath = FPaths::ConvertRelativePathToFull(StagedProjectPath);
+#endif
+	
+	if (FPaths::FileExists(*DevelopmentProjectPath))
 	{
 		// This is a 'in-tree' version of LiveLinkHub, meaning it could be a development build or the version of LLH that is shipped with UnrealEngine.
-		FPaths::SetProjectFilePath(DevelopmentProjectPath);
+		FPaths::SetProjectFilePath(*DevelopmentProjectPath);
 	}
 	else if (FPaths::FileExists(StagedProjectPath))
 	{
