LiveLinkHub - Lower default gamethread framerate to 60 hz
Also remove multithreaded python wrapper command line arg now that the default has changed

#Synced-CL 48513711

diff --git a/Engine/Source/Programs/LiveLinkHub/Source/LiveLinkHubLauncher/Private/LiveLinkHubRun.cpp b/Engine/Source/Programs/LiveLinkHub/Source/LiveLinkHubLauncher/Private/LiveLinkHubRun.cpp
--- a/Engine/Source/Programs/LiveLinkHub/Source/LiveLinkHubLauncher/Private/LiveLinkHubRun.cpp
+++ b/Engine/Source/Programs/LiveLinkHub/Source/LiveLinkHubLauncher/Private/LiveLinkHubRun.cpp
@@ -98,7 +98,7 @@ int32 RunLiveLinkHub(const TCHAR* CommandLine)
 	});
 
 	// Start up the main loop, adding some extra command line arguments:
-	const TCHAR* const ExtraArgs = TEXT("-xrtrackingonly -dpcvars=Engine.Python.ForceWrapperGenerationOnSingleThread=false");
+	const TCHAR* const ExtraArgs = TEXT("-xrtrackingonly");
 	int32 Result = GEngineLoop.PreInit(*FString::Printf(TEXT("%s %s %s"),
 		*FPaths::GetProjectFilePath(), CommandLine, ExtraArgs));
 
@@ -140,6 +140,7 @@ int32 RunLiveLinkHub(const TCHAR* CommandLine)
 		Result = GEngineLoop.Init();
 
 		GEditor->ShouldDisableCPUThrottlingDelegates.Add(UEditorEngine::FShouldDisableCPUThrottling::CreateLambda([]() { return true; }));
+		GEditor->SetMaxFPS(60);
 
 		// Hide the splash screen now that everything is ready to go
 		FPlatformSplash::Hide();
