Set bRequiresPlatformSDK to true for LiveLinkHubLauncher.Build.cs target.

#rnx
#Synced-CL 51630826

diff --git a/Engine/Source/Programs/LiveLinkHub/Source/LiveLinkHubLauncher/LiveLinkHubLauncher.Build.cs b/Engine/Source/Programs/LiveLinkHub/Source/LiveLinkHubLauncher/LiveLinkHubLauncher.Build.cs
--- a/Engine/Source/Programs/LiveLinkHub/Source/LiveLinkHubLauncher/LiveLinkHubLauncher.Build.cs
+++ b/Engine/Source/Programs/LiveLinkHub/Source/LiveLinkHubLauncher/LiveLinkHubLauncher.Build.cs
@@ -97,6 +97,7 @@ public class LiveLinkHubLauncher : ModuleRules
 		}
 
 		ShortName = "LLHub";
+		bRequiresPlatformSDK = true;
 
 		AddRuntimeDependencies();
 	}
