[LiveLinkHub] Enable AvfMedia when the MetaHumanLiveLinkPlugin is enabled

#rnx
#Synced-CL 52569163

diff --git a/Engine/Source/Programs/LiveLinkHub/Source/LiveLinkHub.Target.cs b/Engine/Source/Programs/LiveLinkHub/Source/LiveLinkHub.Target.cs
--- a/Engine/Source/Programs/LiveLinkHub/Source/LiveLinkHub.Target.cs
+++ b/Engine/Source/Programs/LiveLinkHub/Source/LiveLinkHub.Target.cs
@@ -79,7 +79,8 @@ public class LiveLinkHubTarget : TargetRules
 			EnablePlugins.AddRange(new string[]
 			{
 				"MetaHumanLiveLink",
-				"WmfMedia"
+				"WmfMedia",
+				"AvfMedia"
 			});
 		}
 
